UnrealPackageTool: Added a new tool to quickly scan package headers

The initial version only reports the presence of licensee versions as an error, to aid in tracking down invalid engine packages. The goal is to migrate functionality from the PkgInfo commandlet into this tool to provide low-overhead package scanning.

#rnx

diff --git a/Engine/Source/Programs/UnrealPackageTool/UnrealPackageTool.Build.cs b/Engine/Source/Programs/UnrealPackageTool/UnrealPackageTool.Build.cs
new file mode 100644
--- /dev/null
+++ b/Engine/Source/Programs/UnrealPackageTool/UnrealPackageTool.Build.cs
@@ -0,0 +1,18 @@
+// Copyright Epic Games, Inc. All Rights Reserved.
+
+using UnrealBuildTool;
+
+public class UnrealPackageTool : ModuleRules
+{
+	public UnrealPackageTool(ReadOnlyTargetRules Target) : base(Target)
+	{
+		PublicIncludePaths.Add("Runtime/Launch/Public");
+		PrivateIncludePaths.Add("Runtime/Launch/Private");
+
+		PrivateDependencyModuleNames.Add("Core");
+		PrivateDependencyModuleNames.Add("CoreUObject");
+		PrivateDependencyModuleNames.Add("Projects");
+
+		bUseUnity = false;
+	}
+}
diff --git a/Engine/Source/Programs/UnrealPackageTool/UnrealPackageTool.Target.cs b/Engine/Source/Programs/UnrealPackageTool/UnrealPackageTool.Target.cs
new file mode 100644
--- /dev/null
+++ b/Engine/Source/Programs/UnrealPackageTool/UnrealPackageTool.Target.cs
@@ -0,0 +1,26 @@
+// Copyright Epic Games, Inc. All Rights Reserved.
+
+using UnrealBuildTool;
+
+[SupportedPlatforms(UnrealPlatformClass.Desktop)]
+public class UnrealPackageToolTarget : TargetRules
+{
+	public UnrealPackageToolTarget(TargetInfo Target) : base(Target)
+	{
+		Type = TargetType.Program;
+		LinkType = TargetLinkType.Monolithic;
+		LaunchModuleName = "UnrealPackageTool";
+
+		bBuildDeveloperTools = false;
+		bBuildWithEditorOnlyData = true;
+
+		bCompileAgainstApplicationCore = false;
+		bCompileAgainstCoreUObject = true;
+		bCompileAgainstEngine = false;
+		bCompileICU = false;
+
+		bIsBuildingConsoleApplication = true;
+
+		GlobalDefinitions.Add("UE_TRACE_ENABLED=1");
+	}
+}
diff --git a/Engine/Source/Programs/UnrealPackageTool/UnrealPackageTool.cpp b/Engine/Source/Programs/UnrealPackageTool/UnrealPackageTool.cpp
new file mode 100644
--- /dev/null
+++ b/Engine/Source/Programs/UnrealPackageTool/UnrealPackageTool.cpp
@@ -0,0 +1,169 @@
+// Copyright Epic Games, Inc. All Rights Reserved.
+
+#include "Async/ParallelFor.h"
+#include "Containers/Array.h"
+#include "Containers/DepletableMpscQueue.h"
+#include "Containers/StringView.h"
+#include "Containers/UnrealString.h"
+#include "HAL/PlatformFile.h"
+#include "Logging/LogMacros.h"
+#include "Misc/CommandLine.h"
+#include "Misc/PackagePath.h"
+#include "Misc/Parse.h"
+#include "Modules/ModuleManager.h"
+#include "RequiredProgramMainCPPInclude.h"
+#include "String/ParseTokens.h"
+#include "UObject/PackageFileSummary.h"
+
+DEFINE_LOG_CATEGORY_STATIC(LogUnrealPackageTool, Log, All);
+
+IMPLEMENT_APPLICATION(UnrealPackageTool, "UnrealPackageTool");
+
+namespace UE::PackageTool
+{
+
+struct FParams
+{
+	TArray<FString> PackageRoots;
+};
+
+bool ParseParams(const TCHAR* CommandLine, FParams& OutParams)
+{
+	if (!FParse::Param(CommandLine, TEXT("LicenseeVersionIsError")))
+	{
+		UE_LOG(LogUnrealPackageTool, Error,
+			TEXT("Expected -LicenseeVersionIsError which is the only supported mode right now."));
+		return false;
+	}
+
+	OutParams.PackageRoots.Reset();
+
+	for (FString Token; FParse::Token(CommandLine, Token, /*UseEscape*/ false);)
+	{
+		Token.ReplaceInline(TEXT("\""), TEXT(""));
+		const auto GetSwitchValues = [Token = FStringView(Token)](FStringView Match, TArray<FString>& OutValues)
+		{
+			if (Token.StartsWith(Match))
+			{
+				UE::String::ParseTokens(Token.RightChop(Match.Len()), TEXT('+'), [&OutValues](FStringView Value)
+				{
+					OutValues.Emplace(Value);
+				});
+			}
+		};
+		GetSwitchValues(TEXT("-AllPackagesIn="), OutParams.PackageRoots);
+	}
+
+	if (OutParams.PackageRoots.IsEmpty())
+	{
+		UE_LOG(LogUnrealPackageTool, Error,
+			TEXT("Expected at least one package root to be set by -AllPackagesIn=<Root1>+<Root2>."));
+		return false;
+	}
+	for (const FString& PackageRoot : OutParams.PackageRoots)
+	{
+		if (PackageRoot.IsEmpty() || !IFileManager::Get().DirectoryExists(*PackageRoot))
+		{
+			UE_LOG(LogUnrealPackageTool, Error,
+				TEXT("The package root '%s' must point to a directory."), *PackageRoot);
+			return false;
+		}
+	}
+
+	return true;
+}
+
+class FPackageRootVisitor final : public IPlatformFile::FDirectoryVisitor
+{
+public:
+	FPackageRootVisitor()
+		: IPlatformFile::FDirectoryVisitor(EDirectoryVisitorFlags::ThreadSafe)
+	{
+	}
+
+	bool Visit(const TCHAR* Path, bool bIsDirectory) final
+	{
+		if (!bIsDirectory)
+		{
+			EPackageExtension Extension = FPackagePath::ParseExtension(Path);
+			if (Extension == EPackageExtension::Asset || Extension == EPackageExtension::Map)
+			{
+				PackagePaths.Enqueue(Path);
+			}
+		}
+		return true;
+	}
+
+	TDepletableMpscQueue<FString> PackagePaths;
+};
+
+void ScanPackage(const FParams& Params, const TCHAR* Path)
+{
+	if (TUniquePtr<FArchive> Ar{IFileManager::Get().CreateFileReader(Path, FILEREAD_Silent)})
+	{
+		FPackageFileSummary Summary;
+		*Ar << Summary;
+		if (Ar->Close())
+		{
+			if (Summary.CompatibleWithEngineVersion.IsLicenseeVersion())
+			{
+				UE_LOG(LogUnrealPackageTool, Error, TEXT("Package has a licensee version: %s"), Path);
+			}
+		}
+		else
+		{
+			UE_LOG(LogUnrealPackageTool, Warning, TEXT("Failed to read package file summary: %s"), Path);
+		}
+	}
+	else
+	{
+		UE_LOG(LogUnrealPackageTool, Warning, TEXT("Failed to open package: %s"), Path);
+	}
+}
+
+void Main(const FParams& Params)
+{
+	FPackageRootVisitor Visitor;
+
+	ParallelFor(TEXT("ScanPackageRoots.PF"), Params.PackageRoots.Num(), 1, [&Params, &Visitor](int32 Index)
+	{
+		IFileManager::Get().IterateDirectoryRecursively(*Params.PackageRoots[Index], Visitor);
+	}, EParallelForFlags::Unbalanced);
+
+	TArray<FString> PackagePaths;
+	Visitor.PackagePaths.Deplete([&PackagePaths](FString PackagePath)
+	{
+		PackagePaths.Emplace(MoveTemp(PackagePath));
+	});
+
+	ParallelFor(TEXT("ScanPackage.PF"), PackagePaths.Num(), 1, [&Params, &PackagePaths](int32 Index)
+	{
+		ScanPackage(Params, *PackagePaths[Index]);
+	}, EParallelForFlags::Unbalanced);
+}
+
+} // UE::PackageTool
+
+INT32_MAIN_INT32_ARGC_TCHAR_ARGV()
+{
+	UE::PackageTool::FParams Params;
+
+	int32 Ret = GEngineLoop.PreInit(ArgC, ArgV);
+
+	if (Ret == 0)
+	{
+		Ret = !UE::PackageTool::ParseParams(FCommandLine::Get(), Params);
+	}
+
+	if (Ret == 0)
+	{
+		UE::PackageTool::Main(Params);
+	}
+
+	RequestEngineExit(TEXT("Exiting"));
+	FEngineLoop::AppPreExit();
+	FModuleManager::Get().UnloadModulesAtShutdown();
+	FEngineLoop::AppExit();
+
+	return Ret;
+}
