New UnrealPackageTool mode for dumping package info

diff --git a/Engine/Source/Programs/UnrealPackageTool/UnrealPackageTool.Build.cs b/Engine/Source/Programs/UnrealPackageTool/UnrealPackageTool.Build.cs
--- a/Engine/Source/Programs/UnrealPackageTool/UnrealPackageTool.Build.cs
+++ b/Engine/Source/Programs/UnrealPackageTool/UnrealPackageTool.Build.cs
@@ -7,10 +7,19 @@ public class UnrealPackageTool : ModuleRules
 	public UnrealPackageTool(ReadOnlyTargetRules Target) : base(Target)
 	{
 		PublicIncludePathModuleNames.Add("Launch");
-		PrivateDependencyModuleNames.Add("Core");
-		PrivateDependencyModuleNames.Add("CoreUObject");
-		PrivateDependencyModuleNames.Add("Projects");
+		PrivateDependencyModuleNames.AddRange(
+			new string[] {
+				"ApplicationCore", // AssetRegistry requires this 
+				"Core",
+				"CoreUObject",
+				"Projects",
+				"AssetRegistry"
+			}
+		);
 
 		bUseUnity = false;
+		PrivateDependencyModuleNames.Add("CLI11");
+		
+		bEnableExceptions = true;
 	}
 }
diff --git a/Engine/Source/Programs/UnrealPackageTool/UnrealPackageTool.Target.cs b/Engine/Source/Programs/UnrealPackageTool/UnrealPackageTool.Target.cs
--- a/Engine/Source/Programs/UnrealPackageTool/UnrealPackageTool.Target.cs
+++ b/Engine/Source/Programs/UnrealPackageTool/UnrealPackageTool.Target.cs
@@ -11,11 +11,17 @@ public class UnrealPackageToolTarget : TargetRules
 		LinkType = TargetLinkType.Monolithic;
 		IncludeOrderVersion = EngineIncludeOrderVersion.Latest;
 		LaunchModuleName = "UnrealPackageTool";
+		DefaultBuildSettings = BuildSettingsVersion.Latest;
+
+		// Required for CLI11 library 
+		bForceEnableRTTI = true;
+		bForceEnableExceptions = true;
 
 		bBuildDeveloperTools = false;
 		bBuildWithEditorOnlyData = true;
 
-		bCompileAgainstApplicationCore = false;
+		// Required for asset registry module
+		bCompileAgainstApplicationCore = true;
 		bCompileAgainstCoreUObject = true;
 		bCompileAgainstEngine = false;
 		bCompileICU = false;
diff --git a/Engine/Source/Programs/UnrealPackageTool/UnrealPackageTool.cpp b/Engine/Source/Programs/UnrealPackageTool/UnrealPackageTool.cpp
--- a/Engine/Source/Programs/UnrealPackageTool/UnrealPackageTool.cpp
+++ b/Engine/Source/Programs/UnrealPackageTool/UnrealPackageTool.cpp
@@ -1,163 +1,1270 @@
 // Copyright Epic Games, Inc. All Rights Reserved.
 
+#include "Algo/MaxElement.h"
+#include "AssetRegistry/AssetData.h"
+#include "AssetRegistry/PackageReader.h"
 #include "Async/ParallelFor.h"
 #include "Containers/Array.h"
 #include "Containers/ConsumeAllMpmcQueue.h"
 #include "Containers/StringView.h"
 #include "Containers/UnrealString.h"
+#include "HAL/FileManager.h"
 #include "HAL/PlatformFile.h"
-#include "Logging/LogMacros.h"
 #include "Misc/CommandLine.h"
+#include "Misc/ObjectThumbnail.h"
 #include "Misc/PackagePath.h"
 #include "Misc/Parse.h"
+#include "Misc/WildcardString.h"
 #include "Modules/ModuleManager.h"
 #include "RequiredProgramMainCPPInclude.h"
+#include "Serialization/ArchiveProxy.h"
+#include "Serialization/MemoryReader.h"
+#include "Serialization/StructuredArchiveFormatter.h"
+#include "Serialization/Formatters/JsonArchiveOutputFormatter.h"
 #include "String/ParseTokens.h"
 #include "UObject/PackageFileSummary.h"
 
+#include <io.h>
+#include <fcntl.h>
+
+// Undefine legacy macro that conflicts with function names in CLI11
+#undef check
+
+THIRD_PARTY_INCLUDES_START
+#include "CLI/CLI.hpp"
+THIRD_PARTY_INCLUDES_END
+
 DEFINE_LOG_CATEGORY_STATIC(LogUnrealPackageTool, Log, All);
 
 IMPLEMENT_APPLICATION(UnrealPackageTool, "UnrealPackageTool");
 
 namespace UE::PackageTool
 {
+	
+// Convert a utf8 command line parameter to an FString path, treating relative paths as relative to the working directory rather than the engine base dir
+FORCENOINLINE FString ConvertPathParameter(const std::string& InParam)
+{
+	FString Path = UTF8_TO_TCHAR(InParam.c_str());
+	if (FPaths::IsRelative(Path))
+	{
+		FString CWD = FPlatformMisc::LaunchDir();
+		Path = CWD / Path;
+		return FPaths::ConvertRelativePathToFull(MoveTemp(Path));
+	}
+	return Path;
+}
 
-struct FParams
+// Parameters shared between multiple execution modes
+struct FSharedParameters
 {
-	TArray<FString> PackageRoots;
+	bool bJSON = false;
+
+	FSharedParameters(CLI::App* InApp)
+	{
+		InApp->add_flag("--json,-j", bJSON, "Output structured data in JSON format. Without this flag, data is intended to be human readable and not reliably parseable.")
+			->trigger_on_parse(); // Required to allow parsing before subcommands
+	}
 };
 
-bool ParseParams(const TCHAR* CommandLine, FParams& OutParams)
+// Utility archive for printing json/package info to stdout
+struct FArchiveStdOut : public FArchive
 {
-	if (!FParse::Param(CommandLine, TEXT("LicenseeVersionIsError")))
+	int64 Pos = 0;
+
+	~FArchiveStdOut()
 	{
-		UE_LOG(LogUnrealPackageTool, Error,
-			TEXT("Expected -LicenseeVersionIsError which is the only supported mode right now."));
-		return false;
+		fflush(stdout);
+	}
+	
+	// Both formatter types provide utf8 text
+	virtual void Serialize(void* Data, int64 Len) override
+	{
+#if PLATFORM_WINDOWS
+		auto Converted = StringCast<WIDECHAR>((const UTF8CHAR*)Data, Len / sizeof(UTF8CHAR));
+		wprintf(TEXT("%.*s"), (int)(Converted.Length()), Converted.Get());
+		Pos += Converted.Length() * sizeof(TCHAR);
+#else
+		printf("%.*s", Len / sizeof(UTF8CHAR), (UTF8CHAR*)Data);
+		Pos += Len;
+#endif		
 	}
+	
+	// Required for correct formatting of JSON output
+	virtual int64 Tell() override
+	{
+		return Pos;
+	}
+};
 
-	OutParams.PackageRoots.Reset();
+TArray<FString> GatherAssetsInPaths(TConstArrayView<FString> Roots, bool bRecursive)
+{
+	TConsumeAllMpmcQueue<FString> PackagePaths;
+	class FPackageRootVisitor final : public IPlatformFile::FDirectoryVisitor
+	{
+		bool bRecursive;
+		TConsumeAllMpmcQueue<FString>& PackagePaths;
+	public:
+		FPackageRootVisitor(bool InRecursive, TConsumeAllMpmcQueue<FString>& InPackagePaths)
+			: IPlatformFile::FDirectoryVisitor(EDirectoryVisitorFlags::ThreadSafe)
+			, bRecursive(InRecursive)
+			, PackagePaths(InPackagePaths)
+		{
+		}
 
-	for (FString Token; FParse::Token(CommandLine, Token, /*UseEscape*/ false);)
+		bool Visit(const TCHAR* Path, bool bIsDirectory) final
+		{
+			if (!bIsDirectory)
+			{
+				EPackageExtension Extension = FPackagePath::ParseExtension(Path);
+				if (Extension == EPackageExtension::Asset || Extension == EPackageExtension::Map)
+				{
+					PackagePaths.ProduceItem(Path);
+				}
+				return true;
+			}
+			return bRecursive;
+		}
+
+	};
+	
+	FPackageRootVisitor Visitor(bRecursive, PackagePaths);
+	ParallelFor(TEXT("GatherAssestInPaths"), Roots.Num(), 1, [&PackagePaths, &Roots, &Visitor](int32 Index)
 	{
-		Token.ReplaceInline(TEXT("\""), TEXT(""));
-		const auto GetSwitchValues = [Token = FStringView(Token)](FStringView Match, TArray<FString>& OutValues)
+		IFileManager& FM = IFileManager::Get();
+		const FString& Path = Roots[Index];
+		if (FM.FileExists(*Path))
+		{
+			PackagePaths.ProduceItem(Path);
+		}
+		else if (FM.DirectoryExists(*Path))
+		{
+			IFileManager::Get().IterateDirectoryRecursively(*Path, Visitor);
+		}
+		else
 		{
-			if (Token.StartsWith(Match))
+			UE_LOG(LogUnrealPackageTool, Error, TEXT("Input path %s is neither a file nor directory"), *Path);
+		}
+		
+	}, EParallelForFlags::Unbalanced);
+
+	TArray<FString> OutPaths;
+	PackagePaths.ConsumeAllLifo([&OutPaths](FString PackagePath)
+	{
+		OutPaths.Emplace(MoveTemp(PackagePath));
+	});
+
+	return OutPaths;
+}
+
+struct FSubcommand_LicenseeVersionIsError
+{	
+	FSharedParameters& Shared;
+	TArray<FString> PackageRoots;
+
+	FSubcommand_LicenseeVersionIsError(FSharedParameters& InShared, CLI::App* InApp)
+		: Shared(InShared)
+	{
+		CLI::App* Sub = InApp->add_subcommand("LicenseeVersionIsError", "Ensure that no assets have a licensee version set.")
+			->fallthrough()
+			->preparse_callback([this](std::size_t)
 			{
-				UE::String::ParseTokens(Token.RightChop(Match.Len()), TEXT('+'), [&OutValues](FStringView Value)
+				PackageRoots.Reset();
+			});
+		Sub->add_option("--AllPackagesIn", "Check all packages in the given directories")
+			->required()
+			->each([this](const std::string& s) { PackageRoots.Emplace(ConvertPathParameter(s)); })
+			->check(CLI::ExistingDirectory);
+		Sub->parse_complete_callback([this]() { Main(); });
+	}
+	
+	void Main()
+	{
+		UE_LOG(LogUnrealPackageTool, Log, TEXT("Checking packages for licensee version"));
+
+		TArray<FString> PackagePaths = GatherAssetsInPaths(PackageRoots, true);
+
+		ParallelFor(TEXT("ScanPackage.PF"), PackagePaths.Num(), 1, [this, &PackagePaths](int32 Index)
+		{
+			ScanPackage(*PackagePaths[Index]);
+		}, EParallelForFlags::Unbalanced);
+	}
+
+	void ScanPackage(const TCHAR* Path)
+	{
+		if (TUniquePtr<FArchive> Ar{IFileManager::Get().CreateFileReader(Path, FILEREAD_Silent)})
+		{
+			UE_LOG(LogUnrealPackageTool, Log, TEXT("Scanning package %s"), Path);
+			FPackageFileSummary Summary;
+			*Ar << Summary;
+			if (Ar->Close())
+			{
+				if (Summary.CompatibleWithEngineVersion.IsLicenseeVersion())
 				{
-					OutValues.Emplace(Value);
-				});
+					UE_LOG(LogUnrealPackageTool, Error, TEXT("Package has a licensee version: %s"), Path);
+				}
+			}
+			else
+			{
+				UE_LOG(LogUnrealPackageTool, Warning, TEXT("Failed to read package file summary: %s"), Path);
 			}
-		};
-		GetSwitchValues(TEXT("-AllPackagesIn="), OutParams.PackageRoots);
+		}
+		else
+		{
+			UE_LOG(LogUnrealPackageTool, Warning, TEXT("Failed to open package: %s"), Path);
+		}
+	}
+};
+
+// Structured archive formatter to print column-aligned human readable data.
+// Allows us to share driving code with JSON output
+// The output of this formatter is not indended for parsing/machine consumption
+class FTextOutputFormatter final : public FStructuredArchiveFormatter
+{
+	using FBuffer = TUtf8StringBuilder<2048>;
+	static const constexpr int32 IndentWidth = 4;
+	static const constexpr int32 MaxStringLength = 256;
+
+	struct FStackEntry
+	{
+		int32 Indent = 0;
+		TArray<TTuple<FString, FString>> MapEntries;
+		TUniquePtr<FBuffer> ValueBuffer;
+	};
+	FArchive& Ar;
+	TArray<FStackEntry> Stack;
+	
+	int32 GetIndent()
+	{
+		if (Stack.Num() == 0)
+		{
+			return -1; // Prevent top-level record from being indented
+		}
+		return Stack.Top().Indent;
+	}
+	
+	FBuffer& GetValueBuffer()
+	{
+		checkf(Stack.Num() && Stack.Top().ValueBuffer.Get(), TEXT("FTextOutputFormatter: Accessing Value buffer without a scope to output to"));
+		return *Stack.Top().ValueBuffer;
 	}
 
-	if (OutParams.PackageRoots.IsEmpty())
+	void FlushBuffer(FBuffer& Buffer)
 	{
-		UE_LOG(LogUnrealPackageTool, Error,
-			TEXT("Expected at least one package root to be set by -AllPackagesIn=<Root1>+<Root2>."));
+		Ar.Serialize((void*)Buffer.GetData(), Buffer.Len() * sizeof(FBuffer::ElementType));
+		Buffer.Reset();
+	}
+	
+	static bool LooksLikeBinary(FStringView View)
+	{
+		for (const TCHAR& C : View)
+		{
+			// Treat strings containing low-ascii characters as binary
+			if (C < ' ')
+			{
+				return true;
+			}	
+		}
 		return false;
 	}
-	for (const FString& PackageRoot : OutParams.PackageRoots)
+
+	static bool CharNeedsEscaping(TCHAR C)
 	{
-		if (PackageRoot.IsEmpty() || !IFileManager::Get().DirectoryExists(*PackageRoot))
+		switch (C)
 		{
-			UE_LOG(LogUnrealPackageTool, Error,
-				TEXT("The package root '%s' must point to a directory."), *PackageRoot);
-			return false;
+		case '\"':
+		case '\\':
+		case '\b':
+		case '\f':
+		case '\n':
+		case '\r':
+		case '\t':
+			return true;
+		default:
+			if (C < ' ')
+			{
+				return true;
+			}	
 		}
+		return false;
 	}
 
-	return true;
-}
+	static bool NeedsEscaping(FStringView View)
+	{
+		for (const TCHAR& C : View)
+		{
+			if (CharNeedsEscaping(C))
+			{
+				return true;
+			}
+		}
+		return false;
+	}
+	static void WriteEscaped(FBuffer& Buffer, FStringView View)
+	{
+		while(!View.IsEmpty())
+		{
+			bool bEscape = CharNeedsEscaping(View[0]);
+			int32 Count = 1;
+			for (; Count < View.Len() && bEscape == CharNeedsEscaping(View[Count]); ++Count)
+			{
+			}
+			
+			FStringView Section = View.Left(Count);
+			View.RightChopInline(Count);
+			if (bEscape)
+			{
+				for (TCHAR C : Section)
+				{
+					switch (C)
+					{
+					case '\"': Buffer << UTF8TEXTVIEW("\""); break;
+					case '\\': Buffer << UTF8TEXTVIEW("\\"); break;
+					case '\b': Buffer << UTF8TEXTVIEW("\\b"); break;
+					case '\f': Buffer << UTF8TEXTVIEW("\\f"); break;
+					case '\n': Buffer << UTF8TEXTVIEW("\\n"); break;
+					case '\r': Buffer << UTF8TEXTVIEW("\\r"); break;
+					case '\t': Buffer << UTF8TEXTVIEW("\\t"); break;
+					default:
+						if (C < ' ')
+						{
+							Buffer.Appendf("\\u%04x", C);
+						}	
+						else
+						{
+							Buffer.AppendChar(C);
+						}
+					}
+				}	
+			}
+			else
+			{
+				Buffer << Section;
+			}
+		}
+
+	}
 
-class FPackageRootVisitor final : public IPlatformFile::FDirectoryVisitor
-{
 public:
-	FPackageRootVisitor()
-		: IPlatformFile::FDirectoryVisitor(EDirectoryVisitorFlags::ThreadSafe)
+	FTextOutputFormatter(FArchive& InAr)
+		: Ar(InAr)
 	{
+		Ar.SetIsTextFormat(true);
+	}
+	virtual ~FTextOutputFormatter()
+	{
+	}
+
+	virtual FArchive& GetUnderlyingArchive() override
+	{
+		return Ar;
 	}
 
-	bool Visit(const TCHAR* Path, bool bIsDirectory) final
+	virtual bool HasDocumentTree() const override
 	{
-		if (!bIsDirectory)
+		return true;
+	}
+
+	// Records: accumulate field-value pairs and write out so that key names are aligned together
+	virtual void EnterRecord() override
+	{
+		if (Stack.Num() == 0)
 		{
-			EPackageExtension Extension = FPackagePath::ParseExtension(Path);
-			if (Extension == EPackageExtension::Asset || Extension == EPackageExtension::Map)
+			Stack.Push(FStackEntry{ GetIndent() });				
+		}
+		else
+		{
+			Stack.Push(FStackEntry{ GetIndent()+1 });				
+		}
+	}
+	virtual void LeaveRecord() override
+	{
+		FStackEntry Record = Stack.Pop(false);
+		if (Record.MapEntries.Num() == 0)
+		{
+			return;
+		}
+
+		if (Stack.Num() == 0)
+		{
+			// Nothing
+		}
+		else 
+		{
+			FBuffer& Buffer = GetValueBuffer();
+			TTuple<FString, FString>* Longest = Algo::MaxElementBy(Record.MapEntries, [](const TTuple<FString, FString>& Entry) { return Entry.Key.Len(); });
+			int32 LongestKeyLen = Longest->Key.Len();
+			int32 ColumnSize = LongestKeyLen + 4; 
+			for (const TTuple<FString, FString>& KV : Record.MapEntries)
+			{
+				Buffer.Appendf(LINE_TERMINATOR_ANSI "%*s",
+					Record.Indent * IndentWidth, "");
+				Buffer << KV.Key;
+				Buffer.Appendf("%-*s",
+					ColumnSize - KV.Key.Len(), ":");
+				Buffer << KV.Value;
+				if (Stack.Num() == 1)
+				{
+					FlushBuffer(GetValueBuffer());
+				}
+			}	
+			if (Stack.Num() == 1)
 			{
-				PackagePaths.ProduceItem(Path);
+				GetValueBuffer() << LINE_TERMINATOR_ANSI << LINE_TERMINATOR_ANSI;
+				FlushBuffer(GetValueBuffer());
 			}
 		}
+	}
+	virtual void EnterField(FArchiveFieldName Name) override
+	{
+		if (Stack.Num() == 1 )
+		{
+			// Top level fields should be written as major headings 
+			// Then the contents of them are _not_ indented
+			Stack.Top().ValueBuffer = MakeUnique<FBuffer>();
+			GetValueBuffer() << UTF8TEXTVIEW("# ") << Name.Name << LINE_TERMINATOR_ANSI "---";
+		}
+		else
+		{
+			FString ElemName = Name.Name;
+			EnterMapElement(ElemName);
+		}
+	}
+	virtual void LeaveField() override
+	{
+		if (Stack.Num() != 1)
+		{
+			LeaveMapElement();
+		}
+	}
+	virtual bool TryEnterField(FArchiveFieldName Name, bool bEnterWhenSaving) override
+	{
+		EnterField(Name);
 		return true;
 	}
 
-	TConsumeAllMpmcQueue<FString> PackagePaths;
-};
+	virtual void EnterArray(int32& NumElements) override
+	{
+		if (NumElements == 0)
+		{	
+			GetValueBuffer() << UTF8TEXTVIEW("Empty");
+		}
+		EnterStream();
+	}
+	virtual void LeaveArray() override
+	{
+		LeaveStream();
+	}
+	virtual void EnterArrayElement() override
+	{
+		EnterStreamElement();
+	}
+	virtual void LeaveArrayElement() override
+	{
+		LeaveStreamElement();
+	}
 
-void ScanPackage(const FParams& Params, const TCHAR* Path)
-{
-	if (TUniquePtr<FArchive> Ar{IFileManager::Get().CreateFileReader(Path, FILEREAD_Silent)})
+	virtual void EnterStream() override
 	{
-		FPackageFileSummary Summary;
-		*Ar << Summary;
-		if (Ar->Close())
+		if (Stack.Num() == 0)
 		{
-			if (Summary.CompatibleWithEngineVersion.IsLicenseeVersion())
-			{
-				UE_LOG(LogUnrealPackageTool, Error, TEXT("Package has a licensee version: %s"), Path);
-			}
+			checkf(false, TEXT("Top level streams are unsupported"));
 		}
-		else
+		Stack.Push(FStackEntry{ GetIndent() + 1 });	
+	}
+	virtual void LeaveStream() override
+	{
+		if (Stack.Top().ValueBuffer)
+		{
+			LeaveStreamElement();
+		}
+		FStackEntry Record = Stack.Pop(false);
+		if (Stack.Num() == 1)
 		{
-			UE_LOG(LogUnrealPackageTool, Warning, TEXT("Failed to read package file summary: %s"), Path);
+			GetValueBuffer() << LINE_TERMINATOR << LINE_TERMINATOR;
+			FlushBuffer(GetValueBuffer());
 		}
 	}
-	else
+	virtual void EnterStreamElement() override
 	{
-		UE_LOG(LogUnrealPackageTool, Warning, TEXT("Failed to open package: %s"), Path);
+		// Array elements written prefixed with '-' at current indendation level, no other alignment 
+		Stack.Top().ValueBuffer = MakeUnique<FBuffer>();	
+		GetValueBuffer().Appendf(LINE_TERMINATOR_ANSI "%*s- ", Stack.Top().Indent * IndentWidth, "");
 	}
-}
+	virtual void LeaveStreamElement() override
+	{
+		// Write value to parent scope so parent scope can do the alignment it wants 
+		FStackEntry& Parent = Stack[Stack.Num()-2];
+		FStackEntry& This = Stack.Top();
+		(*Parent.ValueBuffer) << This.ValueBuffer->ToView();
+		This.ValueBuffer.Reset();
 
-void Main(const FParams& Params)
-{
-	FPackageRootVisitor Visitor;
+		// If parent scope is top level, flush
+		if (Stack.Num() == 2)
+		{
+			FlushBuffer(*Parent.ValueBuffer);
+		}
+	}
 
-	ParallelFor(TEXT("ScanPackageRoots.PF"), Params.PackageRoots.Num(), 1, [&Params, &Visitor](int32 Index)
+	virtual void EnterMap(int32& NumElements) override
 	{
-		IFileManager::Get().IterateDirectoryRecursively(*Params.PackageRoots[Index], Visitor);
-	}, EParallelForFlags::Unbalanced);
+		if (NumElements == 0)
+		{
+			GetValueBuffer() << UTF8TEXTVIEW("Empty");
+		}
+		EnterRecord();
+	}
+	virtual void LeaveMap() override
+	{
+		LeaveRecord();
+	}
+	virtual void EnterMapElement(FString& Name) override
+	{
+		// Map keys should be aligned like records
+		FStackEntry& Top = Stack.Top();	
+		TTuple<FString, FString>& Entry = Top.MapEntries.AddDefaulted_GetRef();
+		Entry.Key = Name;
+		Top.ValueBuffer = MakeUnique<FBuffer>();
+	}
+	virtual void LeaveMapElement() override
+	{
+		TUniquePtr<FBuffer> Value = MoveTemp(Stack.Top().ValueBuffer);
+		Stack.Top().MapEntries.Last().Value = Value->ToString();
+	}
 
+	virtual void EnterAttributedValue() override{}
+	virtual void EnterAttribute(FArchiveFieldName AttributeName) override{}
+	virtual void EnterAttributedValueValue() override{}
+	virtual void LeaveAttribute() override{}
+	virtual void LeaveAttributedValue() override{}
+	virtual bool TryEnterAttribute(FArchiveFieldName AttributeName, bool bEnterWhenSaving) override
+	{
+		EnterAttribute(AttributeName);
+		return true;
+	}
+
+	virtual bool TryEnterAttributedValueValue() override
+	{
+		EnterAttributedValueValue();
+		return true;
+	}
+
+	virtual void Serialize(uint8& Value) override
+	{
+		GetValueBuffer() << Value;
+	}
+	virtual void Serialize(uint16& Value) override
+	{
+		GetValueBuffer() << Value;
+	}
+	virtual void Serialize(uint32& Value) override
+	{
+		GetValueBuffer() << Value;
+	}
+	virtual void Serialize(uint64& Value) override
+	{
+		GetValueBuffer() << Value;
+	}
+	virtual void Serialize(int8& Value) override
+	{
+		GetValueBuffer() << Value;
+	}
+	virtual void Serialize(int16& Value) override
+	{
+		GetValueBuffer() << Value;
+	}
+	virtual void Serialize(int32& Value) override
+	{
+		GetValueBuffer() << Value;
+	}
+	virtual void Serialize(int64& Value) override
+	{
+		GetValueBuffer() << Value;
+	}
+	virtual void Serialize(float& Value) override
+	{
+		GetValueBuffer().Appendf("%f", Value);
+	}
+	virtual void Serialize(double& Value) override
+	{
+		GetValueBuffer().Appendf("%f", Value);
+	}
+	virtual void Serialize(bool& Value) override
+	{
+		GetValueBuffer() << (Value ? UTF8TEXTVIEW("true") : UTF8TEXTVIEW("false"));
+	}
+	virtual void Serialize(FString& Value) override
+	{
+		if (LooksLikeBinary(FStringView(Value)))
+		{
+			Serialize((void*)*Value, Value.Len() * sizeof(TCHAR));
+			return;
+		}
+
+		FBuffer& Buffer = GetValueBuffer();
+		FStringView View = FStringView(Value).Left(MaxStringLength);
+		Buffer << "\"";
+		if (NeedsEscaping(View))
+		{
+			WriteEscaped(Buffer, View);
+		}	
+		else
+		{
+			Buffer << View;
+		}
+		Buffer << UTF8TEXTVIEW("\"");
+		if (View.Len() < Value.Len())
+		{
+			Buffer << " ...";
+		}
+	}
+	virtual void Serialize(FName& Value) override
+	{
+		GetValueBuffer() << Value;
+	}
+	virtual void Serialize(UObject*& Value) override
+	{
+		FSoftObjectPath Path(Value);
+		Serialize(Path);
+	}
+	virtual void Serialize(FText& Value) override
+	{
+		GetValueBuffer() << UTF8TEXTVIEW("\"") << Value.ToString() << UTF8TEXTVIEW("\"");
+	}
+	virtual void Serialize(FWeakObjectPtr& Value) override
+	{
+		FSoftObjectPath Path;
+		if (UObject* Obj = Value.Get())
+		{
+			Path = FSoftObjectPath(Obj);
+		}
+		Serialize(Path);
+	}
+	virtual void Serialize(FSoftObjectPtr& Value) override
+	{
+		FSoftObjectPath Path = Value.GetUniqueID();
+		Serialize(Path);
+	}
+	virtual void Serialize(FSoftObjectPath& Value) override
+	{
+		GetValueBuffer() << WriteToString<FName::StringBufferSize>(Value);
+	}
+	virtual void Serialize(FLazyObjectPtr& Value) override
+	{
+		FString S = Value.GetUniqueID().ToString();
+		Serialize(S);
+	}
+	virtual void Serialize(FObjectPtr& Value) override
+	{
+		FSoftObjectPath Path = Value.GetPathName();
+		Serialize(Path);
+	}
+	virtual void Serialize(TArray<uint8>& Value) override
+	{
+		Serialize((void*)Value.GetData(), Value.Num());
+	}
+	virtual void Serialize(void* Data, uint64 DataSize) override
+	{
+		FSHAHash Hash;
+		FSHA1::HashBuffer(Data, DataSize, Hash.Hash);
+		GetValueBuffer() << UTF8TEXTVIEW("BINARY ") << DataSize << UTF8TEXTVIEW(" bytes. SHA1 Hash: ");
+		UE::String::BytesToHex(Hash.Hash, GetValueBuffer());
+	}
+};
+
+struct FSubcommand_PackageInfo
+{
+	FSharedParameters& Shared;
 	TArray<FString> PackagePaths;
-	Visitor.PackagePaths.ConsumeAllLifo([&PackagePaths](FString PackagePath)
+	bool bRecursive = false;
+	bool bStdIn = false;
+	bool bWaitForDebugger = false;
+
+	FWildcardString Filter;
+	FString OutPath;
+	bool bOutputToSubdirectories = false;
+	
+	bool bAll = false;
+	bool bSummary = false;
+	bool bNames = false;
+	bool bSoftPaths = false;
+	bool bSoftPackageReferences = false;
+	bool bImports = false;
+	bool bExports = false;
+	bool bText = false;
+	bool bSimple = false;
+	bool bDepends = false;
+	bool bPaths = false;
+	bool bThumbnails = false;
+	bool bLazy = false;
+	bool bAssetRegistry = false;
+	
+	FSubcommand_PackageInfo(FSharedParameters& InShared, CLI::App* InApp)
+		: Shared(InShared)
 	{
-		PackagePaths.Emplace(MoveTemp(PackagePath));
-	});
+		CLI::App* Sub = InApp->add_subcommand("PackageInfo", "Print information contained within asset files.")
+			->fallthrough()
+			->preparse_callback([this](std::size_t)
+			{
+				 PackagePaths.Reset(); 
+				 Filter.Reset();
+			})
+			->parse_complete_callback([this]() { Main(); } );
+		
+		Sub->add_flag("--debug,--wait-for-debugger", bWaitForDebugger, "Wait for a debugger to be attached before continuing");
+		
+		CLI::Option_group* InputGroup = Sub->add_option_group("Input", "Where to get package data from");
+		CLI::Option* PathOption = InputGroup->add_option("-p,--path", "Paths to packages or directories to read." LINE_TERMINATOR_ANSI "Relative paths are treated relative to current working directory.")
+			->each([this](const std::string& s)
+			{
+				PackagePaths.Emplace(ConvertPathParameter(s));
+			});
+		InputGroup->add_flag("-r,--recursive", bRecursive, "Whether to scan directories recursively")
+			->needs(PathOption);
+		InputGroup->add_flag("--stdin", bStdIn, "Whether to read an asset from standard input");
+		InputGroup->required(1);
 
-	ParallelFor(TEXT("ScanPackage.PF"), PackagePaths.Num(), 1, [&Params, &PackagePaths](int32 Index)
+		CLI::Option* OptionOut = Sub->add_option("-o,--out", "Path of file to write output to")
+			->each([this](const std::string& s){ OutPath = ConvertPathParameter(s);});
+		Sub->add_flag("--output-subdirectories", bOutputToSubdirectories, "Write output files to a path structure matching the package names.")
+			->needs(OptionOut);
+		
+		Sub->add_option("-f,--filter", "Wilcard filter to apply to imports, exports, depends map, asset registry outputs")
+			->expected(0,1)
+			->each([this](const std::string& s)
+			{
+				Filter = UTF8_TO_TCHAR(s.c_str());
+			});
+		
+		CLI::Option_group* OutputGroup = Sub->add_option_group("Sections", "Which parts of the package to output information about");
+		OutputGroup->add_flag("--all", bAll, "Write out all package data tables");
+		OutputGroup->add_flag("--summary", bSummary, "Write out the contents of the package file header");
+		OutputGroup->add_flag("--names", bNames, "Write out the contents of the name table");
+		OutputGroup->add_flag("--softpaths", bSoftPaths, "Write out the contents of the soft object path table");
+		OutputGroup->add_flag("--softpackagerefs", bSoftPackageReferences, "Write out the contents of the soft package reference table");
+		OutputGroup->add_flag("--imports", bImports, "Write out the objects imported from other packages by this package");
+		OutputGroup->add_flag("--exports", bExports, "Write out the objects contained within this package");
+		OutputGroup->add_flag("--depends", bDepends, "Write out the contents of the export dependency map");
+		OutputGroup->add_flag("--text", bText, "Write out the contents of the gatherable text (localization) table");
+		OutputGroup->add_flag("--simple", bSimple, "Write a reduced set of information where possible");
+		OutputGroup->add_flag("--thumbnails", bThumbnails, "Write out the contents of the thumbnail table");
+		OutputGroup->add_flag("--assetregistry", bAssetRegistry, "Write out the contenst of the asset registry data (tags etc) in the package.");
+		OutputGroup->required(1);
+	}
+	
+	void Main()
 	{
-		ScanPackage(Params, *PackagePaths[Index]);
-	}, EParallelForFlags::Unbalanced);
-}
+		// If "-waitforattach" or "-WaitForDebugger" was specified, halt startup and wait for a debugger to attach before continuing
+		if (bWaitForDebugger)
+		{
+			while (!FPlatformMisc::IsDebuggerPresent())
+			{
+				FPlatformProcess::Sleep(0.1f);
+			}
+			UE_DEBUG_BREAK();
+		}
+
+		if (!Filter.IsEmpty() && !Filter.ContainsWildcards(*Filter))
+		{
+			Filter = FString::Printf(TEXT("*%s*"), *Filter);
+		}
+
+		TArray<FString> AllPackagePaths;
+		AllPackagePaths.Append(GatherAssetsInPaths(PackagePaths, bRecursive)); 
+		
+		if (bStdIn)
+		{
+			AllPackagePaths.Add(FString{});
+		}
+		
+		IFileManager& FM = IFileManager::Get();
+		if (OutPath.Len() != 0 && AllPackagePaths.Num() > 1 && !FM.DirectoryExists(*OutPath))
+		{
+			if (FM.FileExists(*OutPath))
+			{
+				UE_LOG(LogUnrealPackageTool, Error, TEXT("Cannot output multiple packages to directory %s, a file with that name exists"), *OutPath);
+				return;
+			}
+			else if (!IFileManager::Get().MakeDirectory(*OutPath, true))
+			{
+				UE_LOG(LogUnrealPackageTool, Error, TEXT("Failed to create output directory %s"), *OutPath);
+				return;
+			}
+		}
+		
+		TArray<uint8> StdInBuffer;
+		const FString Extension = Shared.bJSON ? TEXT(".json") : TEXT(".txt");
+		for (FString& PackagePath : AllPackagePaths)
+		{
+			FPackageReader Reader;
+			FPackageReader::EOpenPackageResult ErrorCode;
+			
+			if (PackagePath.Len() == 0)
+			{
+				// Read from standard input
+				StdInBuffer.Reset();
+				_setmode(_fileno(stdin), _O_BINARY);
+				while (!feof(stdin) && !ferror(stdin))
+				{
+					StdInBuffer.Reserve(StdInBuffer.Num() + 1024 * 1024);
+					SIZE_T AmtRead = fread(StdInBuffer.GetData() + StdInBuffer.Num(), 1, StdInBuffer.Max() - StdInBuffer.Num(), stdin);
+					StdInBuffer.AddUninitialized(AmtRead);
+					if (AmtRead == 0)
+					{
+						break;
+					}
+				}
+				TUniquePtr<FMemoryReader> Ar = MakeUnique<FMemoryReader>(StdInBuffer);	
+				Reader.OpenPackageFile(MoveTemp(Ar), &ErrorCode);
+			}
+			else
+			{
+				FPaths::NormalizeFilename(PackagePath);
+				Reader.OpenPackageFile(FStringView(PackagePath), &ErrorCode);
+			}
+			
+			bool bContinue = true;
+			TSet<FGuid> MissingCustomVersions;
+			switch(ErrorCode)
+			{
+				case FPackageReader::EOpenPackageResult::Success:
+					break;
+				case FPackageReader::EOpenPackageResult::CustomVersionMissing:
+					{
+						FCustomVersionArray Versions = Reader.GetPackageFileSummary().GetCustomVersionContainer().GetAllVersions();
+						for (const FCustomVersion& Version : Versions)
+						{
+							TOptional<FCustomVersion> CurrentVersion = FCurrentCustomVersions::Get(Version.Key);
+							if (!CurrentVersion.IsSet())
+							{
+								UE_LOG(LogUnrealPackageTool, Verbose, TEXT("Continuing to load package %s with missing custom version %s"), *PackagePath, *WriteToString<128>(Version.Key));
+								MissingCustomVersions.Add(Version.Key);
+							}
+						}
+					}
+					break;
+				default:
+					bContinue = false;
+					break;
+			}
+
+			if (!bContinue)
+			{
+				UE_LOG(LogUnrealPackageTool, Error, TEXT("Error opening package file %s: %s"), *PackagePath, LexToString(ErrorCode)); 
+				continue;
+			}
+			
+			TUniquePtr<FArchive> Output;
+			if (OutPath.Len() == 0)
+			{
+				Output.Reset(new FArchiveStdOut);
+			}
+			else 
+			{
+				FString OutFilePath;
+				FString InputPath = Reader.GetLongPackageName().IsEmpty() ? PackagePath.Replace(TEXT("\\"), TEXT("/")) : Reader.GetLongPackageName();
+				if (AllPackagePaths.Num() == 1 && !FM.DirectoryExists(*OutPath))
+				{
+					OutFilePath = OutPath;
+				}
+				else if(bOutputToSubdirectories)
+				{
+					// Remove drive letter if necessary
+					if (InputPath.Len() > 2 && FChar::ToUpper(InputPath[0])!=FChar::ToLower(InputPath[0]) && InputPath[1]==TEXT(':'))
+					{
+						InputPath.RightChopInline(2);
+					}
+					OutFilePath = OutPath / InputPath;
+				}
+				else 
+				{
+					// Construct a path from the disk path	
+					FString Path = InputPath;
+					Path.ReplaceCharInline('/', '_');
+					Path.ReplaceCharInline(':', '_');
+					OutFilePath = OutPath / Path + Extension;
+				}
+
+				Output.Reset(FM.CreateFileWriter(*OutFilePath));
+				if (!Output.IsValid())
+				{
+					UE_LOG(LogUnrealPackageTool, Error, TEXT("Error opening output file %s"), *OutFilePath);
+					continue;
+				}
+
+				UTF8CHAR UTF8BOM[] = { (UTF8CHAR)0xEF, (UTF8CHAR)0xBB, (UTF8CHAR)0xBF };
+				Output->Serialize( &UTF8BOM, UE_ARRAY_COUNT(UTF8BOM) * sizeof(UTF8CHAR) );
+			}
+			
+			TUniquePtr<FStructuredArchiveFormatter> Formatter;
+			if (Shared.bJSON)
+			{
+				Formatter.Reset(new FJsonArchiveOutputFormatter(*Output));	
+			}
+			else
+			{
+				Formatter.Reset(new FTextOutputFormatter(*Output));	
+			}
+			FStructuredArchive Writer(*Formatter.Get());
+
+			FPackageFileSummary Summary = Reader.GetPackageFileSummary();
+			FStructuredArchiveRecord Root = Writer.Open().EnterRecord();
+			
+			if (bAll || bSummary)
+			{
+				Root.EnterField(TEXT("PackageFileSummary")) << Summary;
+			}
+
+			if (bAll || bNames)
+			{
+				TArray<FName> Names;
+				if (Reader.GetNames(Names))
+				{
+					Root << SA_VALUE(TEXT("Names"), Names);
+				}
+				else
+				{
+					UE_LOG(LogUnrealPackageTool, Error, TEXT("Error reading name table for package file %s"), *PackagePath);
+				}
+			}
+
+			if (bAll || bSoftPaths )
+			{
+				TArray<FSoftObjectPath> Paths;
+				if (Reader.GetSoftObjectPaths(Paths))
+				{
+					Root << SA_VALUE(TEXT("SoftObjectPaths"), Paths);
+				}
+				else
+				{
+					UE_LOG(LogUnrealPackageTool, Error, TEXT("Error reading soft object path table for package file %s"), *PackagePath);
+				}
+			}
+			
+			FLinkerTables Tables;
+			if (bAll || bImports || bExports || bSoftPackageReferences || bDepends)
+			{
+				if (!Reader.GetImports(Tables.ImportMap))
+				{
+					UE_LOG(LogUnrealPackageTool, Error, TEXT("Error reading import table for package file %s"), *PackagePath);
+				}
+				if (!Reader.GetExports(Tables.ExportMap))
+				{
+					UE_LOG(LogUnrealPackageTool, Error, TEXT("Error reading export table for package file %s"), *PackagePath);
+				}
+				if (!Reader.GetDependsMap(Tables.DependsMap))
+				{
+					UE_LOG(LogUnrealPackageTool, Error, TEXT("Error reading depends map for package file %s"), *PackagePath);
+				}
+				if (!Reader.GetSoftPackageReferenceList(Tables.SoftPackageReferenceList))
+				{
+					UE_LOG(LogUnrealPackageTool, Error, TEXT("Error reading soft package reference list for package file %s"), *PackagePath);
+				}
+			}
+			
+			if (bAll || bImports )
+			{
+				FStructuredArchiveStream Array = Root.EnterField(TEXT("Imports")).EnterStream();
+				for (int32 ImportIndex = 0; ImportIndex < Tables.ImportMap.Num(); ++ImportIndex)
+				{
+					FObjectImport& Import = Tables.ImportMap[ImportIndex];
+					FSoftObjectPath ImportPathName = Tables.GetImportPathName(ImportIndex);
+					if (!Filter.IsEmpty() && !Filter.IsMatch(*WriteToString<FName::StringBufferSize>(ImportPathName)))
+					{
+						continue;
+					}
+					FStructuredArchiveRecord R = Array.EnterElement().EnterRecord();
+					
+					if (bSimple)
+					{
+						R << SA_VALUE(TEXT("PathName"), ImportPathName);
+						FSoftObjectPath ClassPath{ Import.ClassPackage, Import.ClassName, FString{}};
+						R << SA_VALUE(TEXT("ClassPath"), ClassPath);
+					}
+					else
+					{
+						R << SA_VALUE(TEXT("ObjectName"), Import.ObjectName);
+						R << SA_VALUE(TEXT("PathName"), ImportPathName);
+						R << SA_VALUE(TEXT("OuterIndex"), Import.OuterIndex);
+						FSoftObjectPath OuterPathName = Tables.GetImportPathName(Import.OuterIndex);
+						R << SA_VALUE(TEXT("Outer"), OuterPathName);
+						FSoftObjectPath ClassPath{ Import.ClassPackage, Import.ClassName, FString{}};
+						R << SA_VALUE(TEXT("ClassPath"), ClassPath);
+						R << SA_VALUE(TEXT("Optional"), Import.bImportOptional);
+					}
+				}	
+			}
+			
+			FString RootPackageName = Reader.GetLongPackageName();
+			auto PackageIndexToObjectPath = [&Tables, &RootPackageName](FPackageIndex Index) -> FSoftObjectPath
+			{
+				if (Index.IsNull())
+				{
+					return FSoftObjectPath{};
+				}
+
+				return Index.IsExport() 
+					? Tables.GetExportPathName(RootPackageName, Index.ToExport()) 
+					: Tables.GetImportPathName(Index.ToImport());
+			};
+			if (bAll || bExports)
+			{
+				int32 NumExports = Tables.ExportMap.Num();
+				FStructuredArchiveStream Array = Root.EnterField(TEXT("Exports")).EnterStream();
+				for (int32 ExportIndex = 0; ExportIndex < Tables.ExportMap.Num(); ++ExportIndex)
+				{
+					FObjectExport& Export = Tables.ExportMap[ExportIndex];
+					FSoftObjectPath ExportPathName = Tables.GetExportPathName(RootPackageName, ExportIndex);
+					if (!Filter.IsEmpty() && !Filter.IsMatch(*WriteToString<FName::StringBufferSize>(ExportPathName)))
+					{
+						continue;
+					}
+					FStructuredArchiveRecord R = Array.EnterElement().EnterRecord();
+					if (bSimple)
+					{
+						R << SA_VALUE(TEXT("Path"), ExportPathName);
+						FSoftObjectPath ClassPathName = PackageIndexToObjectPath(Export.ClassIndex);
+						R << SA_VALUE(TEXT("Class"), ClassPathName);
+						FString ObjectFlags = LexToString(Export.ObjectFlags);
+						R << SA_VALUE(TEXT("ObjectFlags"), ObjectFlags); 
+						R << SA_VALUE(TEXT("SerialSize"), Export.SerialSize);
+						R << SA_VALUE(TEXT("SerialOffset"), Export.SerialOffset);
+					}
+					else
+					{
+						R << SA_VALUE(TEXT("ObjectName"), Export.ObjectName);
+						R << SA_VALUE(TEXT("OuterIndex"), Export.OuterIndex);
+						R << SA_VALUE(TEXT("Path"), ExportPathName);
+						R << SA_VALUE(TEXT("ClassIndex"), Export.ClassIndex);
+						FSoftObjectPath ClassPathName = PackageIndexToObjectPath(Export.ClassIndex);
+						R << SA_VALUE(TEXT("Class"), ClassPathName);
+						R << SA_VALUE(TEXT("SuperIndex"), Export.SuperIndex);
+						FSoftObjectPath SuperPathName = PackageIndexToObjectPath(Export.SuperIndex);
+						R << SA_VALUE(TEXT("Super"), SuperPathName);
+						R << SA_VALUE(TEXT("TemplateIndex"), Export.TemplateIndex);
+						FSoftObjectPath TemplatePathName = PackageIndexToObjectPath(Export.TemplateIndex);
+						R << SA_VALUE(TEXT("Template"), TemplatePathName);
+						FString ObjectFlags = LexToString(Export.ObjectFlags);
+						R << SA_VALUE(TEXT("ObjectFlags"), ObjectFlags); 
+						R << SA_VALUE(TEXT("HashNext"), Export.HashNext);
+						R << SA_VALUE(TEXT("SerialSize"), Export.SerialSize);
+						R << SA_VALUE(TEXT("SerialOffset"), Export.SerialOffset);
+						R << SA_VALUE(TEXT("ScriptSerializationStartOffset"), Export.ScriptSerializationStartOffset);
+						R << SA_VALUE(TEXT("ScriptSerializationEndOffset"), Export.ScriptSerializationEndOffset);
+						uint8 bForcedExport = Export.bForcedExport;
+						R << SA_VALUE(TEXT("bForcedExport"), bForcedExport);
+						uint8 bNotForClient = Export.bNotForClient; 
+						R << SA_VALUE(TEXT("bNotForClient"), bNotForClient);
+						uint8 bNotForServer = Export.bNotForServer;
+						R << SA_VALUE(TEXT("bNotForServer"), bNotForServer);
+						uint8 bNotAlwaysLoadedForEditorGame = Export.bNotAlwaysLoadedForEditorGame;
+						R << SA_VALUE(TEXT("bNotAlwaysLoadedForEditorGame"), bNotAlwaysLoadedForEditorGame);
+						uint8 bIsAsset = Export.bIsAsset;
+						R << SA_VALUE(TEXT("bIsAsset"), bIsAsset);
+						uint8 bIsInheritedInstance = Export.bIsInheritedInstance;
+						R << SA_VALUE(TEXT("bIsInheritedInstance"), bIsInheritedInstance);
+						uint8 bGeneratePublicHash = Export.bGeneratePublicHash;
+						R << SA_VALUE(TEXT("bGeneratePublicHash"), bGeneratePublicHash);
+					}
+				}
+			}
+			
+			if (bAll || bDepends)
+			{
+				int32 NumExports = Tables.ExportMap.Num();
+				for (int32 ExportIndex = 0; ExportIndex < Tables.ExportMap.Num(); ++ExportIndex)
+				{
+					if (!Filter.IsEmpty() && !Filter.IsMatch(*WriteToString<FName::StringBufferSize>(PackageIndexToObjectPath(FPackageIndex::FromExport(ExportIndex)))))
+					{
+						--NumExports;
+					}
+				}
+
+				FStructuredArchiveMap Map = Root.EnterField(TEXT("DependsMap")).EnterMap(NumExports);
+				for (int32 ExportIndex = 0; ExportIndex < Tables.ExportMap.Num(); ++ExportIndex)
+				{
+					FString ExportPath = PackageIndexToObjectPath(FPackageIndex::FromExport(ExportIndex)).ToString();
+					if (!Filter.IsEmpty() && !Filter.IsMatch(*ExportPath))
+					{
+						continue;
+					}
+					int32 NumDepends = Tables.DependsMap[ExportIndex].Num();
+					if (NumDepends == 0)
+					{
+						continue;
+					}
+					if (bSimple)
+					{
+						FStructuredArchiveArray Array = Map.EnterElement(ExportPath).EnterArray(NumDepends);
+						for (int32 DependsIndex = 0; DependsIndex < NumDepends; ++DependsIndex)
+						{
+							FSoftObjectPath DependsPath = PackageIndexToObjectPath(Tables.DependsMap[ExportIndex][DependsIndex]);
+							Array.EnterElement() << DependsPath;
+						}
+					}	
+					else
+					{
+						FStructuredArchiveArray Array = Map.EnterElement(ExportPath).EnterArray(NumDepends);
+						for (int32 DependsIndex = 0; DependsIndex < NumDepends; ++DependsIndex)
+						{
+							FSoftObjectPath DependsPath = PackageIndexToObjectPath(Tables.DependsMap[ExportIndex][DependsIndex]);
+							FStructuredArchiveRecord R = Array.EnterElement().EnterRecord();
+							R << SA_VALUE(TEXT("Path"), DependsPath);
+							R << SA_VALUE(TEXT("Index"), DependsIndex);
+						}
+					}
+				}
+			}
+			
+			if (bAll || bSoftPackageReferences)
+			{
+				TArray<FName> SoftPackageReferences;
+				if (Reader.GetSoftPackageReferenceList(SoftPackageReferences))
+				{
+					Root << SA_VALUE(TEXT("SoftPackageReferences"), SoftPackageReferences);
+				}
+			}
+
+			if (bAll || bText)
+			{
+				TArray<FGatherableTextData> GatherableTextMap;
+				if (Reader.GetGatherableTextData(GatherableTextMap))
+				{
+					Root << SA_VALUE(TEXT("GatherableTextMap"), GatherableTextMap);
+				}
+				else
+				{
+					UE_LOG(LogUnrealPackageTool, Error, TEXT("Error reading gatherable text data map for package file %s"), *PackagePath);
+				}
+			}
+
+			if (bAll || bThumbnails)
+			{
+				TArray<FObjectFullNameAndThumbnail> Thumbnails;
+				if (Reader.GetThumbnails(Thumbnails))
+				{
+					int32 NumThumbnails = Thumbnails.Num();
+					FStructuredArchiveArray Array = Root.EnterField(TEXT("Thumbnails")).EnterArray(NumThumbnails);
+					for (FObjectFullNameAndThumbnail& Thumbnail : Thumbnails)
+					{
+						FStructuredArchiveRecord R = Array.EnterElement().EnterRecord();
+						R << SA_VALUE(TEXT("ObjectFullName"), Thumbnail.ObjectFullName);
+						R << SA_VALUE(TEXT("FileOffset"), Thumbnail.FileOffset);
+					}
+				}
+				else
+				{
+					UE_LOG(LogUnrealPackageTool, Error, TEXT("Error reading thumbnail map package file %s"), *PackagePath);
+				}
+			}
+			
+			if (bAll || bAssetRegistry)
+			{
+				TArray<FAssetData*> AssetDatas;
+				bool bCookedWithoutAssetData = false;
+				if (Reader.ReadAssetRegistryData(AssetDatas, bCookedWithoutAssetData))
+				{
+					FStructuredArchiveStream Array = Root.EnterField(TEXT("AssetRegistry")).EnterStream();
+					for (FAssetData* AssetData : AssetDatas)
+					{
+						if (!Filter.IsEmpty() && !Filter.IsMatch(*WriteToString<FName::StringBufferSize>(AssetData->AssetName)))
+						{
+							continue;
+						}
+
+						// Note: This could be replaced if/when structured archive serialization is implemented for FAssetData
+						FStructuredArchiveRecord R = Array.EnterElement().EnterRecord();
+						R << SA_VALUE(TEXT("PackageName"), AssetData->PackageName);
+						R << SA_VALUE(TEXT("PackagePath"), AssetData->PackagePath);
+						R << SA_VALUE(TEXT("AssetName"), AssetData->AssetName);
+						FSoftObjectPath AssetClassPath{ AssetData->AssetClassPath, FString{} };
+						R << SA_VALUE(TEXT("AssetClassPath"), AssetClassPath);
+						R << SA_VALUE(TEXT("PackageFlags"), AssetData->PackageFlags);
+						FSoftObjectPath OuterPathName = AssetData->GetOptionalOuterPathName().ToString();
+						R << SA_VALUE(TEXT("OptionalOuterPath"), OuterPathName);
+						TArray<int32> ChunkIDs {AssetData->GetChunkIDs()};
+						R << SA_VALUE(TEXT("ChunkIDs"), ChunkIDs);
+						
+						int32 NumTags = AssetData->TagsAndValues.Num();
+						FStructuredArchiveMap Tags = R.EnterField(TEXT("TagsAndValues")).EnterMap(NumTags);
+						for (const TPair<FName, FAssetTagValueRef>& TagAndValue : AssetData->TagsAndValues)
+						{
+							FString Key = TagAndValue.Key.ToString();
+							FString Value = TagAndValue.Value.AsString();	
+							Tags.EnterElement(Key) << Value;
+						}
+
+						int32 NumBundles = AssetData->TaggedAssetBundles.IsValid() ? AssetData->TaggedAssetBundles->Bundles.Num() : 0;
+						FStructuredArchiveMap Bundles = R.EnterField(TEXT("TaggedAssetBundles")).EnterMap(NumBundles);
+						if (AssetData->TaggedAssetBundles.IsValid())
+						{
+							for (FAssetBundleEntry& Bundle : AssetData->TaggedAssetBundles->Bundles)
+							{
+								FString BundleName = Bundle.BundleName.ToString();
+								Bundles.EnterElement(BundleName) << Bundle.AssetPaths;
+							}
+						}
+						delete AssetData;
+					}	
+				}
+				else
+				{
+					UE_LOG(LogUnrealPackageTool, Error, TEXT("Error reading asset registry data from package file %s"), *PackagePath);
+				}
+			}
+		}
+	}
+};
 
 } // UE::PackageTool
 
 INT32_MAIN_INT32_ARGC_TCHAR_ARGV()
 {
-	UE::PackageTool::FParams Params;
-
+	using namespace UE::PackageTool;
 	int32 Ret = GEngineLoop.PreInit(ArgC, ArgV);
+	CLI::App App(
+		"Utility for reading and modifying with Unreal asset files outside of the editor/engine.\n" \
+		"Copyright Epic Games, Inc. All Rights Reserved.\n",
+
+		"UnrealPackageTool"
+	);
 
-	if (Ret == 0)
+	TUniquePtr<FArchiveStdOut> Out = MakeUnique<FArchiveStdOut>();
+	try 
 	{
-		Ret = !UE::PackageTool::ParseParams(FCommandLine::Get(), Params);
-	}
+		App.ignore_case();
+		App.set_help_all_flag("--help-all", "Expand all help");	
 
-	if (Ret == 0)
+		FSharedParameters Params(&App);
+		FSubcommand_LicenseeVersionIsError LicenseeVersionIsError(Params, &App);
+		FSubcommand_PackageInfo PackageInfo(Params, &App);
+		if (Ret == 0)
+		{
+			App.parse();
+			if (App.get_subcommands().size() == 0)
+			{
+				Out->Logf(TEXT("%s"), UTF8_TO_TCHAR(App.help().c_str()));
+			}
+		}
+	}
+    catch (CLI::CallForAllHelp& e) 
+	{
+		Out->Logf(TEXT("%s"), UTF8_TO_TCHAR(App.help("", CLI::AppFormatMode::All).c_str()));
+		Ret = e.get_exit_code();	
+    }
+	catch (CLI::CallForHelp& e)
+	{
+		Out->Logf(TEXT("%s"), UTF8_TO_TCHAR(App.help().c_str()));
+		Ret = e.get_exit_code();	
+	}
+	catch (CLI::Error& e)
+	{
+		Out->Logf(TEXT("%s"), UTF8_TO_TCHAR(e.what()));
+		Ret = e.get_exit_code();
+	}
+	catch(...)
 	{
-		UE::PackageTool::Main(Params);
+		Out->Logf(TEXT("Unknown error"));
+		Ret = 1;
 	}
 
 	RequestEngineExit(TEXT("Exiting"));
diff --git a/Engine/Source/Runtime/AssetRegistry/Internal/AssetRegistry/PackageReader.h b/Engine/Source/Runtime/AssetRegistry/Internal/AssetRegistry/PackageReader.h
--- a/Engine/Source/Runtime/AssetRegistry/Internal/AssetRegistry/PackageReader.h
+++ b/Engine/Source/Runtime/AssetRegistry/Internal/AssetRegistry/PackageReader.h
@@ -12,6 +12,8 @@ struct FAssetData;
 class FAssetPackageData;
 class FLinkerTables;
 class FPackageDependencyData;
+struct FGatherableTextData;
+struct FObjectFullNameAndThumbnail;
 
 class FPackageReader : public FArchiveUObject
 {
@@ -19,6 +21,7 @@ public:
 	ASSETREGISTRY_API FPackageReader();
 	ASSETREGISTRY_API ~FPackageReader();
 
+	// Note: Keep up to date with LexToString implementation
 	enum class EOpenPackageResult : uint8
 	{
 		/** The package summary loaded successfully */
@@ -51,7 +54,8 @@ public:
 	ASSETREGISTRY_API bool OpenPackageFile(FStringView LongPackageName, FStringView PackageFilename,
 		EOpenPackageResult* OutErrorCode = nullptr);
 	ASSETREGISTRY_API bool OpenPackageFile(FArchive* Loader, EOpenPackageResult* OutErrorCode = nullptr);
-	ASSETREGISTRY_API bool OpenPackageFile(EOpenPackageResult* OutErrorCode = nullptr);
+	ASSETREGISTRY_API bool OpenPackageFile(TUniquePtr<FArchive> Loader, EOpenPackageResult* OutErrorCode = nullptr);
+	ASSETREGISTRY_API bool OpenPackageFile(EOpenPackageResult& OutErrorCode);
 
 	/**
 	 * Returns the LongPackageName from constructor if provided, otherwise calculates it from
@@ -63,14 +67,24 @@ public:
 
 	/** Returns the Summary that was ready by OpenPackageFile. */
 	ASSETREGISTRY_API const FPackageFileSummary& GetPackageFileSummary() const;
+	/** Reads Names if not already read and returns a copy. */
+	ASSETREGISTRY_API bool GetNames(TArray<FName>& OutNames);
 	/** Reads Imports if not already read and returns a copy. */
 	ASSETREGISTRY_API bool GetImports(TArray<FObjectImport>& OutImportMap);
 	/** Reads Exports if not already read and returns a copy. */
 	ASSETREGISTRY_API bool GetExports(TArray<FObjectExport>& OutExportMap);
+	/** Reads DependsMap if not already read and returns a copy. */
+	ASSETREGISTRY_API bool GetDependsMap(TArray<TArray<FPackageIndex>>& OutDependsMap);
 	/** Reads the list of SoftPackageDependencies if not already read and returns a copy. */
 	ASSETREGISTRY_API bool GetSoftPackageReferenceList(TArray<FName>& OutSoftPackageReferenceList);
+	/** Reads the list of SoftObjectPaths if not already read and returns a copy. */
+	ASSETREGISTRY_API bool GetSoftObjectPaths(TArray<FSoftObjectPath>& OutSoftObjectPaths);
 	/** Reads the EditorOnly flags for Imports and Exports and returns them. not cached, reparsed each call. */
 	ASSETREGISTRY_API bool ReadEditorOnlyFlags(TBitArray<>& OutImportUsedInGame, TBitArray<>& OutSoftPackageUsedInGame);
+	/** Reads the list of GatherableTextData if not already read and returns a copy. */
+	ASSETREGISTRY_API bool GetGatherableTextData(TArray<FGatherableTextData>& OutText);
+	/** Reads the list of thumbnails if not already read and returns a copy. */
+	ASSETREGISTRY_API bool GetThumbnails(TArray<FObjectFullNameAndThumbnail>& OutThumbnails);
 
 	/** Reads information from the AR section and converts it to FAssetData. Not cached, reparsed each call. */
 	ASSETREGISTRY_API bool ReadAssetRegistryData(TArray<FAssetData*>& AssetDataList, bool& bOutIsCookedWithoutAssetData);
@@ -122,8 +136,12 @@ private:
 	bool SerializeNameMap();
 	bool SerializeImportMap();
 	bool SerializeExportMap();
+	bool SerializeDependsMap();
 	bool SerializeImportedClasses(const TArray<FObjectImport>& InImportMap, TArray<FName>& OutClassNames);
 	bool SerializeSoftPackageReferenceList();
+	bool SerializeSoftObjectPathMap();
+	bool SerializeGatherableTextDataMap();
+	bool SerializeThumbnailMap();
 	bool SerializeEditorOnlyFlags(TBitArray<>& OutImportUsedInGame, TBitArray<>& OutSoftPackageUsedInGame);
 	bool SerializeSearchableNamesMap(FLinkerTables& OutSearchableNames);
 	bool SerializeAssetRegistryDependencyData(TBitArray<>& OutImportUsedInGame, TBitArray<>& OutSoftPackageUsedInGame,
@@ -149,9 +167,15 @@ private:
 	TArray<FName> NameMap;
 	TArray<FObjectImport> ImportMap;
 	TArray<FObjectExport> ExportMap;
+	TArray<TArray<FPackageIndex>> DependsMap;
 	TArray<FName> SoftPackageReferenceList;
+	TArray<FSoftObjectPath> SoftObjectPathMap;
+	TArray<FGatherableTextData> GatherableTextDataMap;
+	TArray<FObjectFullNameAndThumbnail> ThumbnailMap;
 	int64 PackageFileSize;
 	int64 AssetRegistryDependencyDataOffset;
 	bool bLoaderOwner;
 };
 ENUM_CLASS_FLAGS(FPackageReader::EReadOptions);
+
+ASSETREGISTRY_API const TCHAR* LexToString(FPackageReader::EOpenPackageResult Result);
diff --git a/Engine/Source/Runtime/AssetRegistry/Private/PackageReader.cpp b/Engine/Source/Runtime/AssetRegistry/Private/PackageReader.cpp
--- a/Engine/Source/Runtime/AssetRegistry/Private/PackageReader.cpp
+++ b/Engine/Source/Runtime/AssetRegistry/Private/PackageReader.cpp
@@ -7,6 +7,7 @@
 #include "AssetRegistry/AssetData.h"
 #include "HAL/FileManager.h"
 #include "Internationalization/Internationalization.h"
+#include "Internationalization/GatherableTextData.h"
 #include "Logging/MessageLog.h"
 #include "Misc/Guid.h"
 #include "Misc/PackageName.h"
@@ -17,6 +18,22 @@
 #include "UObject/PackageRelocation.h"
 #include "UObject/PackageTrailer.h"
 
+const TCHAR* LexToString(FPackageReader::EOpenPackageResult Result)
+{
+	switch(Result)
+	{
+		case FPackageReader::EOpenPackageResult::Success: return TEXT("Success");
+		case FPackageReader::EOpenPackageResult::NoLoader: return TEXT("NoLoader");
+		case FPackageReader::EOpenPackageResult::MalformedTag: return TEXT("MalformedTag");
+		case FPackageReader::EOpenPackageResult::VersionTooOld: return TEXT("VersionTooOld");
+		case FPackageReader::EOpenPackageResult::VersionTooNew: return TEXT("VersionTooNew");
+		case FPackageReader::EOpenPackageResult::CustomVersionMissing: return TEXT("CustomVersionMissing");
+		case FPackageReader::EOpenPackageResult::CustomVersionInvalid: return TEXT("CustomVersionInvalid");
+		case FPackageReader::EOpenPackageResult::Unversioned: return TEXT("Unversioned");
+	}
+	checkf(false, TEXT("LexToString unimplemented for EOpenPackageResult %d"), (int32)Result);
+	return TEXT("UNKNOWN");
+}
 
 FPackageReader::FPackageReader()
 	: Loader(nullptr)
@@ -48,7 +65,8 @@ bool FPackageReader::OpenPackageFile(FStringView InLongPackageName, FStringView
 	PackageFilename = InPackageFilename;
 	Loader = IFileManager::Get().CreateFileReader(*PackageFilename);
 	bLoaderOwner = true;
-	return OpenPackageFile(OutErrorCode);
+	EOpenPackageResult Tmp;
+	return OpenPackageFile(OutErrorCode ? *OutErrorCode : Tmp);
 }
 
 bool FPackageReader::OpenPackageFile(FArchive* InLoader, EOpenPackageResult* OutErrorCode)
@@ -58,23 +76,28 @@ bool FPackageReader::OpenPackageFile(FArchive* InLoader, EOpenPackageResult* Out
 	bLoaderOwner = false;
 	LongPackageName.Empty();
 	PackageFilename = Loader->GetArchiveName();
-	return OpenPackageFile(OutErrorCode);
+	EOpenPackageResult Tmp;
+	return OpenPackageFile(OutErrorCode ? *OutErrorCode : Tmp);
 }
 
-bool FPackageReader::OpenPackageFile(EOpenPackageResult* OutErrorCode)
+bool FPackageReader::OpenPackageFile(TUniquePtr<FArchive> InLoader, EOpenPackageResult* OutErrorCode)
 {
-	auto SetPackageErrorCode = [&](const EOpenPackageResult InErrorCode)
-	{
-		if (OutErrorCode)
-		{
-			*OutErrorCode = InErrorCode;
-		}
-	};
+	check(!Loader);
+	Loader = InLoader.Release();
+	bLoaderOwner = true;
+	LongPackageName.Empty();
+	PackageFilename = Loader->GetArchiveName();
+	EOpenPackageResult Tmp;
+	return OpenPackageFile(OutErrorCode ? *OutErrorCode : Tmp);
+}
 
+bool FPackageReader::OpenPackageFile(EOpenPackageResult& OutErrorCode)
+{
+	OutErrorCode = EOpenPackageResult::Success;
 	if( Loader == nullptr )
 	{
 		// Couldn't open the file
-		SetPackageErrorCode(EOpenPackageResult::NoLoader);
+		OutErrorCode = EOpenPackageResult::NoLoader;
 		return false;
 	}
 
@@ -88,7 +111,7 @@ bool FPackageReader::OpenPackageFile(EOpenPackageResult* OutErrorCode)
 	{
 		// Unrecognized or malformed package file
 		UE_LOG(LogAssetRegistry, Error, TEXT("Package %s has malformed tag"), *PackageFilename);
-		SetPackageErrorCode(EOpenPackageResult::MalformedTag);
+		OutErrorCode = EOpenPackageResult::MalformedTag;
 		return false;
 	}
 
@@ -96,7 +119,7 @@ bool FPackageReader::OpenPackageFile(EOpenPackageResult* OutErrorCode)
 	{
 		// Log a warning rather than an error. Linkerload gracefully handles this case.
 		UE_LOG(LogAssetRegistry, Warning, TEXT("Package %s is unversioned which cannot be opened by the current process"), *PackageFilename);
-		SetPackageErrorCode(EOpenPackageResult::Unversioned);
+		OutErrorCode = EOpenPackageResult::Unversioned;
 		return false;
 	}
 
@@ -107,7 +130,7 @@ bool FPackageReader::OpenPackageFile(EOpenPackageResult* OutErrorCode)
 		UE_LOG(	LogAssetRegistry, Warning, TEXT("Package %s is too old. Min Version: %i  Package Version: %i"),
 				*PackageFilename, (int32)VER_UE4_OLDEST_LOADABLE_PACKAGE, PackageFileSummary.GetFileVersionUE().FileVersionUE4);
 
-		SetPackageErrorCode(EOpenPackageResult::Unversioned);
+		OutErrorCode = EOpenPackageResult::Unversioned;
 		return false;
 	}
 
@@ -118,7 +141,7 @@ bool FPackageReader::OpenPackageFile(EOpenPackageResult* OutErrorCode)
 		UE_LOG(	LogAssetRegistry, Warning, TEXT("Package %s is too new. Engine Version: %i  Package Version: %i"),
 				*PackageFilename, GPackageFileUEVersion.ToValue(), PackageFileSummary.GetFileVersionUE().ToValue());
 
-		SetPackageErrorCode(EOpenPackageResult::VersionTooNew);
+		OutErrorCode = EOpenPackageResult::VersionTooNew;
 		return false;
 	}
 
@@ -128,7 +151,7 @@ bool FPackageReader::OpenPackageFile(EOpenPackageResult* OutErrorCode)
 		UE_LOG(LogAssetRegistry, Warning, TEXT("Package %s is too new. Licensee Version: %i Package Licensee Version: %i"),
 			*PackageFilename, GPackageFileLicenseeUEVersion, PackageFileSummary.GetFileVersionLicenseeUE());
 
-		SetPackageErrorCode(EOpenPackageResult::VersionTooNew);
+		OutErrorCode = EOpenPackageResult::VersionTooNew;
 		return false;
 	}
 
@@ -138,20 +161,17 @@ bool FPackageReader::OpenPackageFile(EOpenPackageResult* OutErrorCode)
 	{
 		if (Diff.Type == ECustomVersionDifference::Missing)
 		{
-			SetPackageErrorCode(EOpenPackageResult::CustomVersionMissing);
-			return false;
+			OutErrorCode = EOpenPackageResult::CustomVersionMissing;
 		}
 		else if (Diff.Type == ECustomVersionDifference::Invalid)
 		{
-			SetPackageErrorCode(EOpenPackageResult::CustomVersionInvalid);
-			return false;
+			OutErrorCode = EOpenPackageResult::CustomVersionInvalid;
 		}
 		else if (Diff.Type == ECustomVersionDifference::Newer)
 		{
 			UE_LOG(LogAssetRegistry, Error, TEXT("Package %s has newer custom version of %s"), *PackageFilename, *Diff.Version->GetFriendlyName().ToString());
 
-			SetPackageErrorCode(EOpenPackageResult::VersionTooNew);
-			return false;
+			OutErrorCode = EOpenPackageResult::VersionTooNew;
 		}
 	}
 
@@ -162,11 +182,17 @@ bool FPackageReader::OpenPackageFile(EOpenPackageResult* OutErrorCode)
 
 	const FCustomVersionContainer& PackageFileSummaryVersions = PackageFileSummary.GetCustomVersionContainer();
 	SetCustomVersions(PackageFileSummaryVersions);
+	
+	SetUseUnversionedPropertySerialization((PackageFileSummary.GetPackageFlags() & EPackageFlags::PKG_UnversionedProperties) != 0);
 
 	PackageFileSize = Loader->TotalSize();
+	
+	if (LongPackageName.IsEmpty())
+	{
+		LongPackageName = PackageFileSummary.PackageName;
+	}
 
-	SetPackageErrorCode(EOpenPackageResult::Success);
-	return true;
+	return OutErrorCode == EOpenPackageResult::Success;
 }
 
 bool FPackageReader::TryGetLongPackageName(FString& OutLongPackageName) const
@@ -219,6 +245,16 @@ const FPackageFileSummary& FPackageReader::GetPackageFileSummary() const
 	return PackageFileSummary;
 }
 
+bool FPackageReader::GetNames(TArray<FName>& OutNames)
+{
+	if (!SerializeNameMap())
+	{
+		return false;
+	}	
+	OutNames = NameMap;
+	return true;
+}
+
 bool FPackageReader::GetImports(TArray<FObjectImport>& OutImportMap)
 {
 	if (!SerializeNameMap() || !SerializeImportMap())
@@ -239,6 +275,16 @@ bool FPackageReader::GetExports(TArray<FObjectExport>& OutExportMap)
 	return true;
 }
 
+bool FPackageReader::GetDependsMap(TArray<TArray<FPackageIndex>>& OutDependsMap)
+{
+	if (!SerializeDependsMap())
+	{
+		return false;
+	}
+	OutDependsMap = DependsMap;
+	return true;
+}
+
 bool FPackageReader::GetSoftPackageReferenceList(TArray<FName>& OutSoftPackageReferenceList)
 {
 	if (!SerializeNameMap() || !SerializeSoftPackageReferenceList())
@@ -249,6 +295,37 @@ bool FPackageReader::GetSoftPackageReferenceList(TArray<FName>& OutSoftPackageRe
 	return true;
 }
 
+bool FPackageReader::GetSoftObjectPaths(TArray<FSoftObjectPath>& OutSoftObjectPaths)
+{
+	if (!SerializeNameMap() || !SerializeSoftObjectPathMap())
+	{
+		return false;
+	}
+	OutSoftObjectPaths = SoftObjectPathMap;
+	return true;
+}
+
+bool FPackageReader::GetGatherableTextData(TArray<FGatherableTextData>& OutText)
+{
+	if (!SerializeGatherableTextDataMap())
+	{
+		return false;
+	}
+	OutText = GatherableTextDataMap;
+	return true;
+}
+
+bool FPackageReader::GetThumbnails(TArray<FObjectFullNameAndThumbnail>& OutThumbnails)
+{
+	if (!SerializeThumbnailMap())
+	{
+		return false;
+	}	
+	
+	OutThumbnails = ThumbnailMap;
+	return true;
+}
+
 bool FPackageReader::ReadEditorOnlyFlags(TBitArray<>& OutImportUsedInGame, TBitArray<>& OutSoftPackageUsedInGame)
 {
 	if (!SerializeNameMap() || !SerializeImportMap() || !SerializeSoftPackageReferenceList())
@@ -463,60 +540,31 @@ void FPackageReader::ApplyRelocationToImportMapAndSoftPackageReferenceList(FStri
 
 bool FPackageReader::ReadAssetDataFromThumbnailCache(TArray<FAssetData*>& AssetDataList)
 {
-	if (!StartSerializeSection(PackageFileSummary.ThumbnailTableOffset))
-	{
-		return false;
-	}
-
-	// Determine the package name and path
-	FString PackageName;
-	if (!TryGetLongPackageName(PackageName))
-	{
-		return false;
-	}
-	FString PackagePath = FPackageName::GetLongPackagePath(PackageName);
-
-	// Load the thumbnail count
-	int32 ObjectCount = 0;
-	*this << ObjectCount;
-	const int32 MinBytesPerObject = 1;
-	if (IsError() || ObjectCount < 0 || PackageFileSize < Tell() + ObjectCount * MinBytesPerObject)
+	if (!SerializeThumbnailMap())
 	{
-		UE_PACKAGEREADER_CORRUPTPACKAGE_WARNING("ReadAssetDataFromThumbnailCacheInvalidObjectCount", PackageFilename);
 		return false;
 	}
 
 	// Iterate over every thumbnail entry and harvest the objects classnames
-	for(int32 ObjectIdx = 0; ObjectIdx < ObjectCount; ++ObjectIdx)
+	for(const FObjectFullNameAndThumbnail& Thumbnail : ThumbnailMap)
 	{
-		// Serialize the classname
-		FString AssetClassName;
-		*this << AssetClassName;
-
-		// Serialize the object path.
-		FString ObjectPathWithoutPackageName;
-		*this << ObjectPathWithoutPackageName;
-
-		// Serialize the rest of the data to get at the next object
-		int32 FileOffset = 0;
-		*this << FileOffset;
-
-		if (IsError())
-		{
-			UE_PACKAGEREADER_CORRUPTPACKAGE_WARNING("ReadAssetDataFromThumbnailCacheInvalidObject", PackageFilename);
-			return false;
-		}
-
-		FString GroupNames;
-		FString AssetName;
-
-		if (!ensureMsgf(!ObjectPathWithoutPackageName.Contains(TEXT("."), ESearchCase::CaseSensitive), TEXT("Cannot make FAssetData for sub object %s!"), *ObjectPathWithoutPackageName))
-		{
-			continue;
-		}
+		FString ClassName;
+		FString PackageName;
+		FString ObjectName;
+		FString SubobjectName;
+		FPackageName::SplitFullObjectPath(Thumbnail.ObjectFullName.ToString(), ClassName, PackageName, ObjectName, SubobjectName);
+		FString PackagePath = FPackageName::GetLongPackagePath(PackageName);
 
 		// Create a new FAssetData for this asset and update it with the gathered data
-		AssetDataList.Add(new FAssetData(FName(*PackageName), FName(*PackagePath), FName(*ObjectPathWithoutPackageName), FTopLevelAssetPath(AssetClassName), FAssetDataTagMap(), PackageFileSummary.ChunkIDs, PackageFileSummary.GetPackageFlags()));
+		AssetDataList.Add(new FAssetData(
+			FName(PackageName),
+			FName(PackagePath), 
+			FName(ObjectName), // AssetName
+			FTopLevelAssetPath(ClassName), 
+			FAssetDataTagMap(), 
+			PackageFileSummary.ChunkIDs, 
+			PackageFileSummary.GetPackageFlags())
+		);
 	}
 
 	return true;
@@ -883,6 +931,43 @@ bool FPackageReader::SerializeExportMap()
 	return true;
 }
 
+bool FPackageReader::SerializeDependsMap()
+{
+	if (DependsMap.Num() > 0)
+	{
+		return true;
+	}
+
+	if (PackageFileSummary.DependsOffset > 0 && PackageFileSummary.ExportCount > 0)
+	{
+		if (!StartSerializeSection(PackageFileSummary.DependsOffset))
+		{
+			UE_PACKAGEREADER_CORRUPTPACKAGE_WARNING("SerializeDependsMapInvalidOffset", PackageFilename);
+			return false;
+		}
+
+		const int MinSizePerExport = 1;
+		if (PackageFileSize < Tell() + PackageFileSummary.ExportCount * MinSizePerExport)
+		{
+			UE_PACKAGEREADER_CORRUPTPACKAGE_WARNING("SerializeDependsMapInvalidExportCount", PackageFilename);
+			return false;
+		}
+		DependsMap.Reserve(PackageFileSummary.ExportCount);
+		for (int32 DependsMapIdx = 0; DependsMapIdx < PackageFileSummary.ExportCount; ++DependsMapIdx)
+		{
+			*this << DependsMap.Emplace_GetRef();
+			if (IsError())
+			{
+				UE_PACKAGEREADER_CORRUPTPACKAGE_WARNING("SerializeDependsMapInvalidEntry", PackageFilename);
+				DependsMap.Reset();
+				return false;
+			}
+		}
+	}
+
+	return true;
+}
+
 bool FPackageReader::SerializeSoftPackageReferenceList()
 {
 	if (SoftPackageReferenceList.Num() > 0)
@@ -952,6 +1037,163 @@ bool FPackageReader::SerializeSoftPackageReferenceList()
 	return true;
 }
 
+bool FPackageReader::SerializeSoftObjectPathMap()
+{
+	if (SoftObjectPathMap.Num() > 0)
+	{
+		return true;
+	}
+
+	if (PackageFileSummary.SoftObjectPathsOffset > 0 && PackageFileSummary.SoftObjectPathsCount > 0)
+	{
+		if (!StartSerializeSection(PackageFileSummary.SoftObjectPathsOffset))
+		{
+			UE_PACKAGEREADER_CORRUPTPACKAGE_WARNING("SerializeSoftObjectPathMapListInvalidOffset", PackageFilename);
+			return false;
+		}
+		
+		int32 MinSizePerSoftObjectPath = 0; 
+		if (UEVer() < VER_UE4_ADDED_SOFT_OBJECT_PATH)
+		{
+			MinSizePerSoftObjectPath = 8; // FString
+		}
+		else if (UEVer() < EUnrealEngineObjectUE5Version::FSOFTOBJECTPATH_REMOVE_ASSET_PATH_FNAMES)
+		{
+			MinSizePerSoftObjectPath = 8 + 8; // FName + FString 
+		}
+		else
+		{
+			MinSizePerSoftObjectPath = 8 + 8 + 8; // 2xFName + FString
+		}
+
+		if (PackageFileSize < Tell() + PackageFileSummary.SoftObjectPathsCount * MinSizePerSoftObjectPath)
+		{
+			UE_PACKAGEREADER_CORRUPTPACKAGE_WARNING("SerializeSoftObjectPathMapInvalidCount", PackageFilename);
+			return false;
+		}
+
+		SoftObjectPathMap.Reserve(PackageFileSummary.SoftObjectPathsCount);
+		for (int32 PathIdx = 0; PathIdx < PackageFileSummary.SoftObjectPathsCount; ++PathIdx)
+		{
+			FSoftObjectPath Path;
+			Path.SerializePath(*this);
+			if (IsError())
+			{
+				UE_PACKAGEREADER_CORRUPTPACKAGE_WARNING("SerializeSoftObjectPathMapInvalidPath", PackageFilename);
+				SoftObjectPathMap.Reset();
+				return false;
+			}
+
+			SoftObjectPathMap.Add(Path);
+		}
+	}
+
+	return true;
+}
+
+bool FPackageReader::SerializeGatherableTextDataMap()
+{
+	if (GatherableTextDataMap.Num() > 0)
+	{
+		return true;
+	}
+	
+	if (PackageFileSummary.GatherableTextDataCount > 0 && PackageFileSummary.GatherableTextDataOffset > 0)
+	{
+		if (!StartSerializeSection(PackageFileSummary.GatherableTextDataOffset))
+		{
+			UE_PACKAGEREADER_CORRUPTPACKAGE_WARNING("SerializeGatherableTextDataMapInvalidOffset", PackageFilename);
+			return false;
+		}
+
+		int32 MinSizePerText = 8 + 8 + 4;  // Two FStrings and as an empty array as a lower bound
+		if (PackageFileSize < Tell() + PackageFileSummary.GatherableTextDataCount * MinSizePerText)
+		{
+			UE_PACKAGEREADER_CORRUPTPACKAGE_WARNING("SerializeGatherableTextDataMapInvalidCount", PackageFilename);
+			return false;
+		}
+
+		GatherableTextDataMap.Reset(PackageFileSummary.GatherableTextDataCount);	
+		for (int32 Index = 0; Index < PackageFileSummary.GatherableTextDataCount; ++Index)
+		{
+			FGatherableTextData Data;
+			*this << Data;	
+			if (IsError())
+			{
+				UE_PACKAGEREADER_CORRUPTPACKAGE_WARNING("SerializeGatherableTextDataMapInvalidEntry", PackageFilename);
+				GatherableTextDataMap.Reset();
+				return false;
+			}
+			GatherableTextDataMap.Emplace(MoveTemp(Data));	
+		}
+	}
+
+	return true;
+}
+
+bool FPackageReader::SerializeThumbnailMap()
+{
+	if (ThumbnailMap.Num() > 0)
+	{
+		return true;
+	}
+	
+	if (PackageFileSummary.ThumbnailTableOffset > 0)
+	{
+		if (!StartSerializeSection(PackageFileSummary.ThumbnailTableOffset))
+		{
+			UE_PACKAGEREADER_CORRUPTPACKAGE_WARNING("SerializeThumbnailMapInvalidOffset", PackageFilename);
+			return false;
+		}
+
+		FString PackageName;
+		if (!TryGetLongPackageName(PackageName))
+		{
+			UE_PACKAGEREADER_CORRUPTPACKAGE_WARNING("SerializeThumbnailMapNoPackageName", PackageFilename);
+			return false;
+		}
+
+		int32 NumThumbnails = 0;
+		*this << NumThumbnails;
+		if (IsError() || NumThumbnails < 0)
+		{
+			UE_PACKAGEREADER_CORRUPTPACKAGE_WARNING("SerializeThumbnailMapInvalidCount", PackageFilename);
+			return false;
+		}	
+
+		int32 MinSizePerThumbnail = 8 + 8 + 4;  // Two FStrings and an offset 
+		if (PackageFileSize < Tell() + PackageFileSummary.GatherableTextDataCount * MinSizePerThumbnail)
+		{
+			UE_PACKAGEREADER_CORRUPTPACKAGE_WARNING("SerializeThumbnailMapInvalidCount", PackageFilename);
+			return false;
+		}
+		ThumbnailMap.Reset(NumThumbnails);
+		
+		for (int32 Index=0; Index < NumThumbnails; ++Index)
+		{
+			FString ObjectClassName;
+			*this << ObjectClassName;
+			FString ObjectPathWithoutPackageName;
+			*this << ObjectPathWithoutPackageName;
+			int32 Offset = 0;
+			*this << Offset;
+
+			if (IsError())
+			{
+				UE_PACKAGEREADER_CORRUPTPACKAGE_WARNING("SerializeThumbnailMapInvalidEntry", PackageFilename);
+				return false;
+			}
+			
+			FObjectFullNameAndThumbnail Thumbnail;
+			Thumbnail.ObjectFullName = FName(WriteToString<FName::StringBufferSize>(ObjectClassName, TEXT(" "), PackageName, TEXT("."), ObjectPathWithoutPackageName));
+			Thumbnail.FileOffset = Offset;
+			ThumbnailMap.Emplace(MoveTemp(Thumbnail));
+		}
+	}
+
+	return true;
+}
+
 bool FPackageReader::SerializeEditorOnlyFlags(TBitArray<>& OutImportUsedInGame, TBitArray<>& OutSoftPackageUsedInGame)
 {
 	if (AssetRegistryDependencyDataOffset == INDEX_NONE &&
