Fix newline behavior on windows to not output \r\r\n on stdout which confuses some shells.

diff --git a/Engine/Source/Programs/UnrealPackageTool/UnrealPackageTool.cpp b/Engine/Source/Programs/UnrealPackageTool/UnrealPackageTool.cpp
--- a/Engine/Source/Programs/UnrealPackageTool/UnrealPackageTool.cpp
+++ b/Engine/Source/Programs/UnrealPackageTool/UnrealPackageTool.cpp
@@ -80,6 +80,14 @@ struct FArchiveStdOut : public FArchive
 	virtual void Serialize(void* Data, int64 Len) override
 	{
 #if PLATFORM_WINDOWS
+		// replace \r with a space to avoid CRT printf's function expanding \r\n to \r\r\n
+		for (UTF8CHAR* C = (UTF8CHAR*)Data, *End = C + Len / sizeof(UTF8CHAR); C != End; ++C)
+		{
+			if (*C == '\r')
+			{
+				*C = (UTF8CHAR)' ';
+			}
+		}
 		auto Converted = StringCast<WIDECHAR>((const UTF8CHAR*)Data, Len / sizeof(UTF8CHAR));
 		wprintf(TEXT("%.*s"), (int)(Converted.Length()), Converted.Get());
 		Pos += Converted.Length() * sizeof(TCHAR);
