Fix handling multiple paths on command line, add alias of -d for searching directories in LicenseeVersionIsError mode

e.g. LicenseeVersionIsError -d E:\p4\ws_1\Engine -d E:\p4\ws_1\EngineTest

diff --git a/Engine/Source/Programs/UnrealPackageTool/UnrealPackageTool.cpp b/Engine/Source/Programs/UnrealPackageTool/UnrealPackageTool.cpp
--- a/Engine/Source/Programs/UnrealPackageTool/UnrealPackageTool.cpp
+++ b/Engine/Source/Programs/UnrealPackageTool/UnrealPackageTool.cpp
@@ -178,8 +178,9 @@ struct FSubcommand_LicenseeVersionIsError
 			{
 				PackageRoots.Reset();
 			});
-		Sub->add_option("--AllPackagesIn", "Check all packages in the given directories")
+		Sub->add_option("--AllPackagesIn,-d", "Check all packages in the given directories")
 			->required()
+			->multi_option_policy(CLI::MultiOptionPolicy::TakeAll)
 			->each([this](const std::string& s) { PackageRoots.Emplace(ConvertPathParameter(s)); })
 			->check(CLI::ExistingDirectory);
 		Sub->parse_complete_callback([this]() { Main(); });
@@ -722,6 +723,7 @@ struct FSubcommand_PackageInfo
 		
 		CLI::Option_group* InputGroup = Sub->add_option_group("Input", "Where to get package data from");
 		CLI::Option* PathOption = InputGroup->add_option("-p,--path", "Paths to packages or directories to read." LINE_TERMINATOR_ANSI "Relative paths are treated relative to current working directory.")
+			->multi_option_policy(CLI::MultiOptionPolicy::TakeAll)
 			->each([this](const std::string& s)
 			{
 				PackagePaths.Emplace(ConvertPathParameter(s));
