- Fixed UnrealPackageTool to work on Mac (and probably linux)
#rnx

diff --git a/Engine/Source/Programs/UnrealPackageTool/UnrealPackageTool.cpp b/Engine/Source/Programs/UnrealPackageTool/UnrealPackageTool.cpp
--- a/Engine/Source/Programs/UnrealPackageTool/UnrealPackageTool.cpp
+++ b/Engine/Source/Programs/UnrealPackageTool/UnrealPackageTool.cpp
@@ -24,8 +24,10 @@
 #include "String/ParseTokens.h"
 #include "UObject/PackageFileSummary.h"
 
+#if PLATFORM_WINDOWS
 #include <io.h>
 #include <fcntl.h>
+#endif
 
 // Undefine legacy macro that conflicts with function names in CLI11
 #undef check
@@ -92,7 +94,7 @@ struct FArchiveStdOut : public FArchive
 		wprintf(TEXT("%.*s"), (int)(Converted.Length()), Converted.Get());
 		Pos += Converted.Length() * sizeof(TCHAR);
 #else
-		printf("%.*s", Len / sizeof(UTF8CHAR), (UTF8CHAR*)Data);
+		printf("%.*s", (int)(Len / sizeof(char)), (const char*)Data);
 		Pos += Len;
 #endif		
 	}
@@ -812,7 +814,11 @@ struct FSubcommand_PackageInfo
 			{
 				// Read from standard input
 				StdInBuffer.Reset();
+#if PLATFORM_WINDOWS
 				_setmode(_fileno(stdin), _O_BINARY);
+#else
+				freopen(nullptr, "rb", stdin);
+#endif
 				while (!feof(stdin) && !ferror(stdin))
 				{
 					StdInBuffer.Reserve(StdInBuffer.Num() + 1024 * 1024);
