Silence other logging so it doesn't interfere with parsing output.

diff --git a/Engine/Source/Programs/UnrealPackageTool/UnrealPackageTool.cpp b/Engine/Source/Programs/UnrealPackageTool/UnrealPackageTool.cpp
--- a/Engine/Source/Programs/UnrealPackageTool/UnrealPackageTool.cpp
+++ b/Engine/Source/Programs/UnrealPackageTool/UnrealPackageTool.cpp
@@ -1237,6 +1237,11 @@ INT32_MAIN_INT32_ARGC_TCHAR_ARGV()
 {
 	using namespace UE::PackageTool;
 	int32 Ret = GEngineLoop.PreInit(ArgC, ArgV);
+	
+	// Disable all logging because we want to output to stdout
+	FSelfRegisteringExec::StaticExec(nullptr, TEXT("log unrealpackagetool only"), *GLog);
+	FSelfRegisteringExec::StaticExec(nullptr, TEXT("log unrealpackagetool log"), *GLog);
+
 	CLI::App App(
 		"Utility for reading and modifying with Unreal asset files outside of the editor/engine.\n" \
 		"Copyright Epic Games, Inc. All Rights Reserved.\n",
