Fix log channel manipulation to use correct names to enable LogUnrealPackageTool.

diff --git a/Engine/Source/Programs/UnrealPackageTool/UnrealPackageTool.cpp b/Engine/Source/Programs/UnrealPackageTool/UnrealPackageTool.cpp
--- a/Engine/Source/Programs/UnrealPackageTool/UnrealPackageTool.cpp
+++ b/Engine/Source/Programs/UnrealPackageTool/UnrealPackageTool.cpp
@@ -1237,10 +1237,10 @@ INT32_MAIN_INT32_ARGC_TCHAR_ARGV()
 {
 	using namespace UE::PackageTool;
 	int32 Ret = GEngineLoop.PreInit(ArgC, ArgV);
-	
+
 	// Disable all logging because we want to output to stdout
-	FSelfRegisteringExec::StaticExec(nullptr, TEXT("log unrealpackagetool only"), *GLog);
-	FSelfRegisteringExec::StaticExec(nullptr, TEXT("log unrealpackagetool log"), *GLog);
+	FSelfRegisteringExec::StaticExec(nullptr, TEXT("log logunrealpackagetool only"), *GLog);
+	FSelfRegisteringExec::StaticExec(nullptr, TEXT("log logunrealpackagetool log"), *GLog);
 
 	CLI::App App(
 		"Utility for reading and modifying with Unreal asset files outside of the editor/engine.\n" \
