Fix compilation of UnrealPackageTool .
It has to be treated as an engine module in order to use internal headers.
Because of its compilation setup it also needs to declare some macros that would otherwise be declared by UBT for IMPLEMENT_APPLICATION.

diff --git a/Engine/Source/Programs/UnrealPackageTool/UnrealPackageTool.Build.cs b/Engine/Source/Programs/UnrealPackageTool/UnrealPackageTool.Build.cs
--- a/Engine/Source/Programs/UnrealPackageTool/UnrealPackageTool.Build.cs
+++ b/Engine/Source/Programs/UnrealPackageTool/UnrealPackageTool.Build.cs
@@ -21,5 +21,6 @@ public class UnrealPackageTool : ModuleRules
 		PrivateDependencyModuleNames.Add("CLI11");
 		
 		bEnableExceptions = true;
+		bTreatAsEngineModule = true; // So we can use internal headers
 	}
 }
diff --git a/Engine/Source/Programs/UnrealPackageTool/UnrealPackageTool.cpp b/Engine/Source/Programs/UnrealPackageTool/UnrealPackageTool.cpp
--- a/Engine/Source/Programs/UnrealPackageTool/UnrealPackageTool.cpp
+++ b/Engine/Source/Programs/UnrealPackageTool/UnrealPackageTool.cpp
@@ -38,6 +38,10 @@ THIRD_PARTY_INCLUDES_END
 
 DEFINE_LOG_CATEGORY_STATIC(LogUnrealPackageTool, Log, All);
 
+// These macros are not properly defined by UBT in the case of an engine program with bTreatAsEngineModule=true
+// So define them here as a workaround
+#define IMPLEMENT_ENCRYPTION_KEY_REGISTRATION()
+#define IMPLEMENT_SIGNING_KEY_REGISTRATION()
 IMPLEMENT_APPLICATION(UnrealPackageTool, "UnrealPackageTool");
 
 namespace UE::PackageTool
