CIS fix.

#rnx

diff --git a/Engine/Source/Programs/UnrealPackageTool/UnrealPackageTool.cpp b/Engine/Source/Programs/UnrealPackageTool/UnrealPackageTool.cpp
--- a/Engine/Source/Programs/UnrealPackageTool/UnrealPackageTool.cpp
+++ b/Engine/Source/Programs/UnrealPackageTool/UnrealPackageTool.cpp
@@ -997,7 +997,7 @@ struct FSubcommand_PackageInfo
 					if (bSimple)
 					{
 						R << SA_VALUE(TEXT("PathName"), ImportPathName);
-						FSoftObjectPath ClassPath{ Import.ClassPackage, Import.ClassName, FString{}};
+						FSoftObjectPath ClassPath{ FTopLevelAssetPath(Import.ClassPackage, Import.ClassName) };
 						R << SA_VALUE(TEXT("ClassPath"), ClassPath);
 					}
 					else
@@ -1007,7 +1007,7 @@ struct FSubcommand_PackageInfo
 						R << SA_VALUE(TEXT("OuterIndex"), Import.OuterIndex);
 						FSoftObjectPath OuterPathName = Tables.GetImportPathName(Import.OuterIndex);
 						R << SA_VALUE(TEXT("Outer"), OuterPathName);
-						FSoftObjectPath ClassPath{ Import.ClassPackage, Import.ClassName, FString{}};
+						FSoftObjectPath ClassPath{ FTopLevelAssetPath(Import.ClassPackage, Import.ClassName) };
 						R << SA_VALUE(TEXT("ClassPath"), ClassPath);
 						R << SA_VALUE(TEXT("Optional"), Import.bImportOptional);
 					}
