Adding some AssetRegistry tools to UnrealPackageTool

#Synced-CL 47063821

diff --git a/Engine/Source/Programs/UnrealPackageTool/UPT_AssetRegistry.cpp b/Engine/Source/Programs/UnrealPackageTool/UPT_AssetRegistry.cpp
new file mode 100644
--- /dev/null
+++ b/Engine/Source/Programs/UnrealPackageTool/UPT_AssetRegistry.cpp
@@ -0,0 +1,605 @@
+// Copyright Epic Games, Inc. All Rights Reserved.
+
+#include "UnrealPackageTool.h"
+
+#include "AssetRegistry/AssetRegistryState.h"
+#include "HAL/FileManager.h"
+#include "HAL/PlatformMisc.h"
+#include "Misc/ScopeExit.h"
+#include "Serialization/StructuredArchiveFormatter.h"
+#include "Serialization/Formatters/JsonArchiveOutputFormatter.h"
+#include "SQLiteDatabase.h"
+#include "SQLiteTypes.h"
+
+// Undefine legacy macro that conflicts with function names in CLI11
+#undef check
+
+THIRD_PARTY_INCLUDES_START
+#include "CLI/CLI.hpp"
+THIRD_PARTY_INCLUDES_END
+
+namespace UE::PackageTool
+{
+struct FOutput
+{
+	TUniquePtr<FArchive> Archive;
+
+	FArchive* operator->()
+	{
+		return Archive.Get();
+	}
+
+	FArchive& operator*()
+	{
+		return *Archive.Get();
+	}
+};
+
+struct FStructuredOutput
+{
+	FOutput UnstructuredOutput;
+	TUniquePtr<FStructuredArchiveFormatter> Formatter;
+	FStructuredArchive Writer;
+
+	FStructuredArchive* operator->()
+	{
+		return &Writer;
+	}
+
+	FStructuredArchive& operator*()
+	{
+		return Writer;
+	}
+};
+
+void WriteString(FStringView S, FArchive& Ar)
+{
+	auto Converted = StringCast<UTF8CHAR>(S.GetData(), S.Len());
+	Ar.Serialize((void*)GetData(Converted), GetNum(Converted) * sizeof(UTF8CHAR));
+}
+
+enum class EOutputType
+{
+	Stdout,
+	File,
+	Directory
+};
+
+class FSubcommand_AssetRegistry : public FSubcommand
+{
+	FSharedParameters* Shared = nullptr;
+
+	// Path to input asset registry state as a .bin file
+	FString BinInput;
+
+	EOutputType OutputType = EOutputType::Stdout;
+	// Path to output file or directory
+	FString OutputPath;
+
+	bool bDump = false;
+
+	struct FPrintOptions
+	{
+		bool bAssetNames = false;
+		bool bTags = false;
+	} PrintOptions;;
+
+	struct FSQLiteOptions
+	{
+	} SQLiteOptions;
+
+	struct FDumpOptions
+	{
+		bool bAll = false;
+		bool bByObjectPath = false;
+		bool bByPackageName = false;
+		bool bByPath = false;
+		bool bByClass = false;
+		bool bByTag = false;
+		bool bAssetTags = false;
+		bool bDependencies = false;
+		bool bDependencyDetails = false;
+		bool bLegacyDependencies = false;
+		bool bPackageData = false;
+		bool bAssetBundles = false;
+	} DumpOptions;
+
+	virtual void RegisterSubcommand(FSharedParameters* InShared, CLI::App* App) override;
+	void Main_Print();
+	void Main_SQLite();
+#if ASSET_REGISTRY_STATE_DUMPING_ENABLED
+	void Main_Dump();
+#endif // ASSET_REGISTRY_STATE_DUMPING_ENABLED
+	
+	// Helper to create output files. Filename will be used when the output type is directory
+	FOutput CreateOutput(FStringView Filename);
+	FStructuredOutput CreateStructuredOutput(FStringView Filename);
+	
+	FAssetRegistryState LoadAssetRegistryState();
+};
+
+void FSubcommand_AssetRegistry::RegisterSubcommand(FSharedParameters* InShared, CLI::App* App)
+{
+	Shared = InShared;
+	CLI::App* Sub = App->add_subcommand("AssetRegistry", "Read asset registry bin/cache files and modify or output that data");
+
+	auto AddInputGroup = [this](CLI::App* To)
+	{
+		CLI::Option_group* InputGroup = To->add_option_group("Input", "Where to get asset registry data data from");
+		InputGroup->require_option(1, 1);
+		InputGroup->add_option("--bin",
+			"Path to a .bin file containing a saved asset registry, e.g. the output DevelopmentAssetRegistry.bin file from the cooker")
+			->each([this](const std::string& s) { BinInput = ConvertPathParameter(s); });
+	};
+
+	auto AddOutputGroup = [this](CLI::App* To)
+	{ 
+		CLI::Option_group* OutputGroup = To->add_option_group("Output", "Where to write output data to");
+		// CLI::Option_group* OutputGroup = Sub->add_option_group("Output",);
+		OutputGroup->require_option(1, 1);
+		OutputGroup->add_option("--out-file,-o", "Write all output to the destination file")
+			->each([this](const std::string& s) 
+			{
+				OutputType = EOutputType::File;
+				OutputPath = ConvertPathParameter(s); 
+			});
+		OutputGroup->add_option("--out-dir", "Write all output to files in the destination directory")
+			->each([this](const std::string& s) 
+			{
+				OutputType = EOutputType::Directory;
+				OutputPath = ConvertPathParameter(s); 
+			});
+		OutputGroup->add_flag("--stdout", "Write all output to standard output")->each([this](const std::string& s) 
+		{
+			OutputType = EOutputType::Stdout; 
+		});
+	};
+
+	CLI::App* Sub_Print = Sub->add_subcommand("Print", "Print out some data from the asset registry")
+		->final_callback([this]() { Main_Print(); });
+	Sub_Print->add_flag("--asset-names", PrintOptions.bAssetNames, "Print all asset names");
+	Sub_Print->add_flag("--tags", PrintOptions.bTags, "Print all known asset registry tags");
+	Sub_Print->require_option();
+	AddInputGroup(Sub_Print);
+	AddOutputGroup(Sub_Print);
+
+
+	CLI::App* Sub_SQLite = Sub->add_subcommand("SQLite", "Create or update an sqlite database with asset registry data")
+		->final_callback([this]() { Main_SQLite(); })
+		->footer("Note that this is an experimental feature and the table structure of the database is subject to change.");
+	AddInputGroup(Sub_SQLite);
+	
+	Sub_SQLite->add_option("--out-file,-o", "Output file for the database")
+	 	->required()
+		->each([this](const std::string& s) 
+		{
+			 OutputType = EOutputType::File;
+			 OutputPath = ConvertPathParameter(s); 
+		});
+
+#if ASSET_REGISTRY_STATE_DUMPING_ENABLED
+	CLI::App* Sub_Dump = Sub->add_subcommand("Dump", "Dump the complete asset registry")
+		->final_callback([this]() { Main_Dump(); });
+	AddInputGroup(Sub_Dump);
+	AddOutputGroup(Sub_Dump);
+
+	Sub_Dump->add_flag("--all", DumpOptions.bAll,"Output all information in the asset registry");
+	Sub_Dump->add_flag("--by-object-path", DumpOptions.bByObjectPath,"Output the path of each asset in the asset registry");
+	Sub_Dump->add_flag("--by-package-name", DumpOptions.bByPackageName, "Output each asset in the asset registry grouped by package name");
+	Sub_Dump->add_flag("--by-path", DumpOptions.bByPath, "Output each asset in the asset registry grouped by package path (package name without everything past the last /)");
+	Sub_Dump->add_flag("--by-class", DumpOptions.bByClass,"Output each asset in the asset registry grouped by class");
+	Sub_Dump->add_flag("--by-tag", DumpOptions.bByTag,"Output each asset in the asset registry grouped by asset tag (i.e. assets may appear more than once)");
+	Sub_Dump->add_flag("--asset-tags", DumpOptions.bAssetTags,"Output each asset in the asset registry and the tags it has");
+	Sub_Dump->add_flag("--dependencies", DumpOptions.bDependencies, "Output each dependency node in the asset registry without listing all dependencies from it");
+	Sub_Dump->add_flag("--dependency-details", DumpOptions.bDependencyDetails,"Output each dependency node in the asset registry and list all basic dependencies from it");
+	Sub_Dump->add_flag("--legacy-dependencies", DumpOptions.bLegacyDependencies, "Output each dependency node in the asset registry and list all detailed dependencies from it");
+	Sub_Dump->add_flag("--package-data", DumpOptions.bPackageData,"Output information about each package file in the asset registry");
+	Sub_Dump->add_flag("--asset-bundles", DumpOptions.bAssetBundles,"Output the contents of each asset bundle in the asset registry");
+#endif // ASSET_REGISTRY_STATE_DUMPING_ENABLED
+}
+
+FOutput FSubcommand_AssetRegistry::CreateOutput(FStringView Filename)
+{
+	IFileManager& FM = IFileManager::Get();
+	TUniquePtr<FArchive> Output;
+	switch (OutputType)
+	{
+	case EOutputType::Stdout:
+		Output.Reset(new FArchiveStdOut);
+		break;
+	case EOutputType::File:
+		{
+			Output.Reset(FM.CreateFileWriter(*OutputPath));
+		}
+		break;
+	case EOutputType::Directory:
+		{
+			if (!FM.DirectoryExists(*OutputPath))
+			{
+				if (!FM.MakeDirectory(*OutputPath, true))
+				{
+					throw CLI::Error("FailedToCreateOutDirectory", "Failed to create output directory", 1);
+				}
+			}
+
+			TStringBuilder<512> FilePath;
+			FPathViews::Append(FilePath, OutputPath, Filename);
+			Output.Reset(FM.CreateFileWriter(*FilePath));
+		}
+		break;
+	}
+
+	if (!Output.Get())
+	{
+		throw CLI::Error("FailedToOpenOutFile", "Failed to open output file", 1);
+	}
+
+	return FOutput{ MoveTemp(Output) };
+}
+
+FStructuredOutput FSubcommand_AssetRegistry::CreateStructuredOutput(FStringView Filename)
+{
+	FOutput Output = CreateOutput(Filename);
+	TUniquePtr<FStructuredArchiveFormatter> Formatter;
+	if (Shared->bJSON)
+	{
+		Formatter.Reset(new FJsonArchiveOutputFormatter(*Output));	
+	}
+	else
+	{
+		Formatter.Reset(new FTextOutputFormatter(*Output));	
+	}
+
+	FStructuredArchiveFormatter* FormatterPtr = Formatter.Get();
+	return FStructuredOutput{ MoveTemp(Output), MoveTemp(Formatter), FStructuredArchive(*FormatterPtr) };
+}
+
+FAssetRegistryState FSubcommand_AssetRegistry::LoadAssetRegistryState()
+{
+	FAssetRegistryState State;
+	FAssetRegistryLoadOptions LoadOptions;
+	LoadOptions.ParallelWorkers = FPlatformMisc::NumberOfCoresIncludingHyperthreads();
+	FAssetRegistryVersion::Type Version = FAssetRegistryVersion::PreVersioning;
+	if (!BinInput.IsEmpty())
+	{
+		if (!FAssetRegistryState::LoadFromDisk(*BinInput, LoadOptions, State, &Version))
+		{
+			throw CLI::Error("FailedToLoadAssetRegistry", "Failed to load asset registry state", 1);
+		}
+	}
+
+	return MoveTemp(State);
+}
+
+void FSubcommand_AssetRegistry::Main_SQLite()
+{
+	// Re-enable some logging since we aren't using stdout
+	FSelfRegisteringExec::StaticExec(nullptr, TEXT("log reset"), *GLog);
+
+	IFileManager& FM = IFileManager::Get();
+	if (FM.FileExists(*OutputPath))
+	{
+		FM.Delete(*OutputPath);
+	}
+
+	FSQLiteDatabase DB;
+	DB.Open(*OutputPath);
+
+	if (!DB.IsValid())
+	{
+		throw CLI::Error("FailedToOpenDB", "Failed to open database");
+	}
+
+	ON_SCOPE_EXIT
+	{
+		DB.Close();
+	};
+
+	// don't bother waiting for disk to flush (let the OS handle it)
+	DB.Execute(TEXT("PRAGMA synchronous=OFF;"));
+	// 1000 page cache
+	DB.Execute(TEXT("PRAGMA cache_size=1000;"));
+	// large-ish page sizes for modern disks
+	DB.Execute(TEXT("PRAGMA page_size=65535;"));
+	// hold the file lock the whole time.
+	DB.Execute(TEXT("PRAGMA locking_mode=EXCLUSIVE;"));
+	// Hold things in memory unti commit(?)
+	DB.Execute(TEXT("PRAGMA temp_store=MEMORY;"));
+
+	DB.Execute(
+		TEXT("CREATE TABLE AssetClasses (")
+		TEXT("PackageName TEXT, ")
+		TEXT("ClassName TEXT, ")
+		TEXT("UNIQUE (PackageName, ClassName) ")
+		TEXT(");")
+	);
+
+	DB.Execute(
+		TEXT("CREATE TABLE Assets (")
+		TEXT("PackageName TEXT, ")
+		TEXT("AssetName TEXT, ")
+		TEXT("AssetClassID INTEGER, ")
+		TEXT("OptionalOuterPath TEXT, ")
+		TEXT("PackageFlags INTEGER, ") // TODO: Move to a separate table? 
+		TEXT("FOREIGN KEY(AssetClassID) REFERENCES AssetClasses(ROWID), ")
+		TEXT("UNIQUE (PackageName, AssetName) ")
+		TEXT(");")
+	);
+
+	DB.Execute(
+		TEXT("CREATE TABLE Tags (")
+		TEXT("TagName TEXT, ")
+		TEXT("UNIQUE (TagName) ")
+		TEXT(");")
+	);
+
+	DB.Execute(
+		TEXT("CREATE TABLE AssetTags (")
+		TEXT("AssetID INTEGER, ")
+		TEXT("TagID INTEGER, ")
+		TEXT("TagValue TEXT, ")
+		TEXT("FOREIGN KEY(AssetID) REFERENCES Assets(ROWID), ")
+		TEXT("FOREIGN KEY(TagID) REFERENCES Tags(ROWID), ")
+		TEXT("UNIQUE(AssetID, TagID) ")
+		TEXT(");")
+	);
+
+	DB.Execute(
+		TEXT("CREATE TEMPORARY TABLE AssetsImport( ")
+		TEXT("PackageName, AssetName, ClassPackageName, ClassName, OptionalOuterPath, PackageFlags, ")
+		TEXT("UNIQUE(PackageName, AssetName) ")
+		TEXT(");")
+	);
+
+	DB.Execute(
+		TEXT("CREATE TEMPORARY TABLE TagsImport( ")
+		TEXT("PackageName, AssetName, TagName, TagValue, ")
+		TEXT("UNIQUE(PackageName, AssetName, TagName) ")
+		TEXT(");")
+	);
+
+	FSQLitePreparedStatement InsertClass;
+	const TCHAR* InsertClassText = TEXT("INSERT INTO AssetClasses (PackageName, ClassName) VALUES (?1, ?2);");
+	if (!InsertClass.Create(DB, InsertClassText))
+	{
+		throw CLI::Error("FailedToCreatePreparedStatement", "Failed to create prepared statement", 1);
+	}
+
+	FSQLitePreparedStatement InsertAsset;
+	const TCHAR* InsertAssetText = 
+		TEXT("INSERT INTO AssetsImport (PackageName, AssetName, ClassPackageName, ClassName, OptionalOuterPath, PackageFlags) ")
+		TEXT("VALUES (?1, ?2, ?3, ?4, ?5, ?6);");
+	if (!InsertAsset.Create(DB, InsertAssetText))
+	{
+		throw CLI::Error("FailedToCreatePreparedStatement", "Failed to create prepared statement", 1);
+	}
+
+	FSQLitePreparedStatement InsertTag;
+	const TCHAR* InsertTagText = TEXT("INSERT INTO TagsImport (PackageName, AssetName, TagName, TagValue) VALUES (?1, ?2, ?3, ?4);");
+	if (!InsertTag.Create(DB, InsertTagText))
+	{
+		throw CLI::Error("FailedToCreatePreparedStatement", "Failed to create prepared statement", 1);
+	}
+
+
+	const TCHAR* FinalizeAssets = 
+	 	// Populate Assets table referencing AssetClasses
+		TEXT("INSERT INTO Assets (PackageName, AssetName, AssetClassID, OptionalOuterPath, PackageFlags) ")
+		TEXT("SELECT AssetsImport.PackageName, AssetsImport.AssetName, AssetClasses.ROWID, AssetsImport.OptionalOuterPath, AssetsImport.PackageFlags ")
+		TEXT("FROM AssetsImport ")
+		TEXT("INNER JOIN AssetClasses ON AssetClasses.PackageName == AssetsImport.ClassPackageName AND AssetClasses.ClassName == AssetsImport.ClassName; ");
+	const TCHAR* FinalizeTags = 
+		// Populate Tags
+		TEXT("INSERT INTO Tags (TagName) SELECT DISTINCT TagName FROM TagsImport; ");
+	const TCHAR* FinalizeAssetTags =
+		// Populate AssetTags referencing Assets and Tags
+		TEXT("INSERT INTO AssetTags (AssetID, TagID, TagValue) ")	
+		TEXT("SELECT Assets.ROWID, Tags.RowID, TagsImport.TagValue ")
+		TEXT("FROM TagsImport ")
+		TEXT("INNER JOIN Assets ON Assets.PackageName == TagsImport.PackageName AND Assets.AssetName == TagsImport.AssetName ")
+		TEXT("INNER JOIN Tags ON Tags.TagName == TagsImport.TagName; ");
+
+	FAssetRegistryState State = LoadAssetRegistryState();
+
+	TSet<FTopLevelAssetPath> AssetClasses;
+	State.EnumerateAllAssets([&](const FAssetData& AssetData) 
+	{
+		AssetClasses.Add(AssetData.AssetClassPath);
+	});
+
+
+	DB.Execute(TEXT("BEGIN TRANSACTION;"));
+	int32 i = 0;
+	constexpr int32 BatchSize = 10000;
+	for (FTopLevelAssetPath Class : AssetClasses)
+	{
+		InsertClass.Reset();
+		InsertClass.ClearBindings();
+		InsertClass.SetBindingValueByIndex(1, Class.GetPackageName().ToString());
+		InsertClass.SetBindingValueByIndex(2, Class.GetAssetName().ToString());
+
+		if (!InsertClass.Execute())
+		{
+			throw CLI::Error("FailedToInsertAssetClass", "Failed to insert asset class", 1);
+		}
+
+		if (++i == BatchSize)
+		{
+			i = 0;
+			DB.Execute(TEXT("END TRANSACTION;"));
+			DB.Execute(TEXT("BEGIN TRANSACTION;"));
+		}
+
+	}
+	DB.Execute(TEXT("END TRANSACTION;"));
+
+
+	i = 0;
+	DB.Execute(TEXT("BEGIN TRANSACTION;"));
+	State.EnumerateAllAssets([&](const FAssetData& AssetData) 
+	{
+		InsertAsset.Reset();
+		InsertAsset.ClearBindings();
+		InsertAsset.SetBindingValueByIndex(1, AssetData.PackageName.ToString());
+		InsertAsset.SetBindingValueByIndex(2, AssetData.AssetName.ToString());
+		InsertAsset.SetBindingValueByIndex(3, AssetData.AssetClassPath.GetPackageName().ToString());
+		InsertAsset.SetBindingValueByIndex(4, AssetData.AssetClassPath.GetAssetName().ToString());
+		InsertAsset.SetBindingValueByIndex(5, AssetData.GetOptionalOuterPathName().ToString());
+		InsertAsset.SetBindingValueByIndex(6, AssetData.PackageFlags);
+
+		if (!InsertAsset.Execute())
+		{
+			throw CLI::Error("FailedToInsertAsset", "Failed to insert asset", 1);
+		}
+
+		if (++i == BatchSize)
+		{
+			i = 0;
+			DB.Execute(TEXT("END TRANSACTION;"));
+			DB.Execute(TEXT("BEGIN TRANSACTION;"));
+		}
+	});
+	DB.Execute(TEXT("END TRANSACTION;"));
+
+	i = 0;
+	DB.Execute(TEXT("BEGIN TRANSACTION;"));
+	State.EnumerateAllAssets([&](const FAssetData& AssetData) 
+	{
+		for (TPair<FName, FAssetTagValueRef> Tag : AssetData.TagsAndValues)
+		{
+			InsertTag.Reset();
+			InsertTag.ClearBindings();
+			InsertTag.SetBindingValueByIndex(1, AssetData.PackageName.ToString());
+			InsertTag.SetBindingValueByIndex(2, AssetData.AssetName.ToString());
+			InsertTag.SetBindingValueByIndex(3, Tag.Key.ToString());
+			InsertTag.SetBindingValueByIndex(4, Tag.Value.AsString());
+
+			if (!InsertTag.Execute())
+			{
+				throw CLI::Error("FailedToInsertTag", "Failed to insert tag", 1);
+			}
+
+			if (++i == BatchSize)
+			{
+				i = 0;
+				DB.Execute(TEXT("END TRANSACTION;"));
+				DB.Execute(TEXT("BEGIN TRANSACTION;"));
+			}
+		}
+	});
+	DB.Execute(TEXT("END TRANSACTION;"));
+
+	DB.Execute(FinalizeAssets);
+	DB.Execute(FinalizeTags);
+	DB.Execute(FinalizeAssetTags);
+}
+
+void FSubcommand_AssetRegistry::Main_Print()
+{
+	FAssetRegistryState State = LoadAssetRegistryState();
+
+	FStructuredOutput Output = CreateStructuredOutput(TEXTVIEW("AssetRegistryData.txt"));
+	FStructuredArchive::FRecord Root = Output->Open().EnterRecord();
+	if (PrintOptions.bAssetNames)
+	{
+		FStructuredArchive::FStream Array = Root.EnterField(TEXT("AssetNames")).EnterStream();
+		State.EnumerateAllAssets([&](const FAssetData& AssetData) 
+		{
+			FSoftObjectPath Path = AssetData.GetSoftObjectPath();
+			Array << Path;
+			return true;
+		});
+	}
+	if (PrintOptions.bTags)
+	{
+		FStructuredArchive::FStream Array = Root.EnterField(TEXT("Tags")).EnterStream();
+		State.EnumerateTags([&](FName TagName) 
+		{
+			Array << TagName;
+			return true;
+		});
+	}
+
+}
+
+#if ASSET_REGISTRY_STATE_DUMPING_ENABLED
+void FSubcommand_AssetRegistry::Main_Dump()
+{
+	FAssetRegistryState State = LoadAssetRegistryState();
+
+	TArray<FString> Pages;
+	
+	// TODO: Move args directly to option parsing
+	TArray<FString> Args;
+	if (DumpOptions.bAll)
+	{
+		Args.Emplace(TEXT("All"));
+	}
+
+	if (DumpOptions.bByObjectPath)
+	{
+		Args.Emplace(TEXT("ObjectPath"));
+	}
+
+	if (DumpOptions.bByPackageName)
+	{
+		Args.Emplace(TEXT("PackageName"));
+	}
+
+	if (DumpOptions.bByPath)
+	{
+		Args.Emplace(TEXT("Path"));
+	}
+
+	if (DumpOptions.bByClass)
+	{
+		Args.Emplace(TEXT("Class"));
+	}
+
+	if (DumpOptions.bByTag)
+	{
+		Args.Emplace(TEXT("Tag"));
+	}
+
+	if (DumpOptions.bAssetTags)
+	{
+		Args.Emplace(TEXT("AssetTags"));
+	}
+
+	if (DumpOptions.bDependencies)
+	{
+		Args.Emplace(TEXT("Dependencies"));
+	}
+
+	if (DumpOptions.bDependencyDetails)
+	{
+		Args.Emplace(TEXT("DependencyDetails"));
+	}
+
+	if (DumpOptions.bLegacyDependencies)
+	{
+		Args.Emplace(TEXT("LegacyDependencies"));
+	}
+
+	if (DumpOptions.bPackageData)
+	{
+		Args.Emplace(TEXT("PackageData"));
+	}
+
+	if (DumpOptions.bAssetBundles)
+	{
+		Args.Emplace(TEXT("AssetBundles"));
+	}
+
+	State.Dump(Args, Pages, 10000);
+	for (int32 i=0; i < Pages.Num(); ++i)
+	{
+		FOutput Output = CreateOutput(WriteToString<256>(TEXT("AssetRegistryDump_"), i, TEXT(".txt")));
+		WriteString(Pages[i], *Output.Archive);
+	}
+}
+#endif // ASSET_REGISTRY_STATE_DUMPING_ENABLED
+
+FSubcommand_AssetRegistry GSubcommand_AssetRegistry;
+
+}
\ No newline at end of file
diff --git a/Engine/Source/Programs/UnrealPackageTool/UnrealPackageTool.Build.cs b/Engine/Source/Programs/UnrealPackageTool/UnrealPackageTool.Build.cs
--- a/Engine/Source/Programs/UnrealPackageTool/UnrealPackageTool.Build.cs
+++ b/Engine/Source/Programs/UnrealPackageTool/UnrealPackageTool.Build.cs
@@ -13,6 +13,7 @@ public class UnrealPackageTool : ModuleRules
 				"Core",
 				"CoreUObject",
 				"Projects",
+				"SQLiteCore",
 				"AssetRegistry"
 			}
 		);
diff --git a/Engine/Source/Programs/UnrealPackageTool/UnrealPackageTool.cpp b/Engine/Source/Programs/UnrealPackageTool/UnrealPackageTool.cpp
--- a/Engine/Source/Programs/UnrealPackageTool/UnrealPackageTool.cpp
+++ b/Engine/Source/Programs/UnrealPackageTool/UnrealPackageTool.cpp
@@ -1,5 +1,7 @@
 // Copyright Epic Games, Inc. All Rights Reserved.
 
+#include "UnrealPackageTool.h"
+
 #include "Algo/MaxElement.h"
 #include "AssetRegistry/AssetData.h"
 #include "AssetRegistry/PackageReader.h"
@@ -38,6 +40,7 @@ THIRD_PARTY_INCLUDES_END
 
 DEFINE_LOG_CATEGORY_STATIC(LogUnrealPackageTool, Log, All);
 
+
 // These macros are not properly defined by UBT in the case of an engine program with bTreatAsEngineModule=true
 // So define them here as a workaround
 #define IMPLEMENT_ENCRYPTION_KEY_REGISTRATION()
@@ -46,9 +49,23 @@ IMPLEMENT_APPLICATION(UnrealPackageTool, "UnrealPackageTool");
 
 namespace UE::PackageTool
 {
+
+FSubcommand* SubcommandList = nullptr;
+
+FSubcommand::FSubcommand()
+{
+	Next = SubcommandList;
+	SubcommandList = this;
+}
+
+FSubcommand::~FSubcommand()
+{
+
+}
+	
 	
 // Convert a utf8 command line parameter to an FString path, treating relative paths as relative to the working directory rather than the engine base dir
-FORCENOINLINE FString ConvertPathParameter(const std::string& InParam)
+FString ConvertPathParameter(const std::string& InParam)
 {
 	FString Path = UTF8_TO_TCHAR(InParam.c_str());
 	if (FPaths::IsRelative(Path))
@@ -61,54 +78,44 @@ FORCENOINLINE FString ConvertPathParameter(const std::string& InParam)
 }
 
 // Parameters shared between multiple execution modes
-struct FSharedParameters
+FSharedParameters::FSharedParameters(CLI::App* InApp)
 {
-	bool bJSON = false;
-
-	FSharedParameters(CLI::App* InApp)
-	{
-		InApp->add_flag("--json,-j", bJSON, "Output structured data in JSON format. Without this flag, data is intended to be human readable and not reliably parseable.")
-			->trigger_on_parse(); // Required to allow parsing before subcommands
-	}
-};
+	InApp->add_flag("--json,-j", bJSON, "Output structured data in JSON format. Without this flag, data is intended to be human readable and not reliably parseable.")
+		->trigger_on_parse(); // Required to allow parsing before subcommands
+}
 
 // Utility archive for printing json/package info to stdout
-struct FArchiveStdOut : public FArchive
+FArchiveStdOut::~FArchiveStdOut()
 {
-	int64 Pos = 0;
+	fflush(stdout);
+}
 
-	~FArchiveStdOut()
-	{
-		fflush(stdout);
-	}
-	
-	// Both formatter types provide utf8 text
-	virtual void Serialize(void* Data, int64 Len) override
-	{
+// Both formatter types provide utf8 text
+void FArchiveStdOut::Serialize(void* Data, int64 Len) 
+{
 #if PLATFORM_WINDOWS
-		// replace \r with a space to avoid CRT printf's function expanding \r\n to \r\r\n
-		for (UTF8CHAR* C = (UTF8CHAR*)Data, *End = C + Len / sizeof(UTF8CHAR); C != End; ++C)
+	// replace \r with a space to avoid CRT printf's function expanding \r\n to \r\r\n
+	for (UTF8CHAR* C = (UTF8CHAR*)Data, *End = C + Len / sizeof(UTF8CHAR); C != End; ++C)
+	{
+		if (*C == '\r')
 		{
-			if (*C == '\r')
-			{
-				*C = (UTF8CHAR)' ';
-			}
+			*C = (UTF8CHAR)' ';
 		}
-		auto Converted = StringCast<WIDECHAR>((const UTF8CHAR*)Data, Len / sizeof(UTF8CHAR));
-		wprintf(TEXT("%.*s"), (int)(Converted.Length()), Converted.Get());
-		Pos += Converted.Length() * sizeof(TCHAR);
+	}
+	auto Converted = StringCast<WIDECHAR>((const UTF8CHAR*)Data, Len / sizeof(UTF8CHAR));
+	wprintf(TEXT("%.*s"), (int)(Converted.Length()), Converted.Get());
+	Pos += Converted.Length() * sizeof(TCHAR);
 #else
-		printf("%.*s", (int)(Len / sizeof(char)), (const char*)Data);
-		Pos += Len;
+	printf("%.*s", (int)(Len / sizeof(char)), (const char*)Data);
+	Pos += Len;
 #endif		
-	}
-	
-	// Required for correct formatting of JSON output
-	virtual int64 Tell() override
-	{
-		return Pos;
-	}
-};
+}
+
+// Required for correct formatting of JSON output
+int64 FArchiveStdOut::Tell() 
+{
+	return Pos;
+}
 
 TArray<FString> GatherAssetsInPaths(TConstArrayView<FString> Roots, bool bRecursive)
 {
@@ -230,465 +237,456 @@ struct FSubcommand_LicenseeVersionIsError
 	}
 };
 
-// Structured archive formatter to print column-aligned human readable data.
-// Allows us to share driving code with JSON output
-// The output of this formatter is not indended for parsing/machine consumption
-class FTextOutputFormatter final : public FStructuredArchiveFormatter
+int32 FTextOutputFormatter::GetIndent()
 {
-	using FBuffer = TUtf8StringBuilder<2048>;
-	static const constexpr int32 IndentWidth = 4;
-	static const constexpr int32 MaxStringLength = 256;
-
-	struct FStackEntry
-	{
-		int32 Indent = 0;
-		TArray<TTuple<FString, FString>> MapEntries;
-		TUniquePtr<FBuffer> ValueBuffer;
-	};
-	FArchive& Ar;
-	TArray<FStackEntry> Stack;
-	
-	int32 GetIndent()
+	if (Stack.Num() == 0)
 	{
-		if (Stack.Num() == 0)
-		{
-			return -1; // Prevent top-level record from being indented
-		}
-		return Stack.Top().Indent;
-	}
-	
-	FBuffer& GetValueBuffer()
-	{
-		checkf(Stack.Num() && Stack.Top().ValueBuffer.Get(), TEXT("FTextOutputFormatter: Accessing Value buffer without a scope to output to"));
-		return *Stack.Top().ValueBuffer;
+		return -1; // Prevent top-level record from being indented
 	}
+	return Stack.Top().Indent;
+}
 
-	void FlushBuffer(FBuffer& Buffer)
-	{
-		Ar.Serialize((void*)Buffer.GetData(), Buffer.Len() * sizeof(FBuffer::ElementType));
-		Buffer.Reset();
-	}
-	
-	static bool LooksLikeBinary(FStringView View)
-	{
-		for (const TCHAR& C : View)
-		{
-			// Treat strings containing low-ascii characters as binary
-			if (C < ' ')
-			{
-				return true;
-			}	
-		}
-		return false;
-	}
+FTextOutputFormatter::FBuffer& FTextOutputFormatter::GetValueBuffer()
+{
+	checkf(Stack.Num() && Stack.Top().ValueBuffer.Get(), TEXT("FTextOutputFormatter: Accessing Value buffer without a scope to output to"));
+	return *Stack.Top().ValueBuffer;
+}
 
-	static bool CharNeedsEscaping(TCHAR C)
-	{
-		switch (C)
-		{
-		case '\"':
-		case '\\':
-		case '\b':
-		case '\f':
-		case '\n':
-		case '\r':
-		case '\t':
-			return true;
-		default:
-			if (C < ' ')
-			{
-				return true;
-			}	
-		}
-		return false;
-	}
+void FTextOutputFormatter::FlushBuffer(FBuffer& Buffer)
+{
+	Ar.Serialize((void*)Buffer.GetData(), Buffer.Len() * sizeof(FBuffer::ElementType));
+	Buffer.Reset();
+}
 
-	static bool NeedsEscaping(FStringView View)
-	{
-		for (const TCHAR& C : View)
-		{
-			if (CharNeedsEscaping(C))
-			{
-				return true;
-			}
-		}
-		return false;
-	}
-	static void WriteEscaped(FBuffer& Buffer, FStringView View)
+bool FTextOutputFormatter::LooksLikeBinary(FStringView View)
+{
+	for (const TCHAR& C : View)
 	{
-		while(!View.IsEmpty())
+		// Treat strings containing low-ascii characters as binary
+		if (C < ' ')
 		{
-			bool bEscape = CharNeedsEscaping(View[0]);
-			int32 Count = 1;
-			for (; Count < View.Len() && bEscape == CharNeedsEscaping(View[Count]); ++Count)
-			{
-			}
-			
-			FStringView Section = View.Left(Count);
-			View.RightChopInline(Count);
-			if (bEscape)
-			{
-				for (TCHAR C : Section)
-				{
-					switch (C)
-					{
-					case '\"': Buffer << UTF8TEXTVIEW("\""); break;
-					case '\\': Buffer << UTF8TEXTVIEW("\\"); break;
-					case '\b': Buffer << UTF8TEXTVIEW("\\b"); break;
-					case '\f': Buffer << UTF8TEXTVIEW("\\f"); break;
-					case '\n': Buffer << UTF8TEXTVIEW("\\n"); break;
-					case '\r': Buffer << UTF8TEXTVIEW("\\r"); break;
-					case '\t': Buffer << UTF8TEXTVIEW("\\t"); break;
-					default:
-						if (C < ' ')
-						{
-							Buffer.Appendf("\\u%04x", C);
-						}	
-						else
-						{
-							Buffer.AppendChar(C);
-						}
-					}
-				}	
-			}
-			else
-			{
-				Buffer << Section;
-			}
-		}
-
-	}
-
-public:
-	FTextOutputFormatter(FArchive& InAr)
-		: Ar(InAr)
-	{
-		Ar.SetIsTextFormat(true);
-	}
-	virtual ~FTextOutputFormatter()
-	{
-	}
-
-	virtual FArchive& GetUnderlyingArchive() override
-	{
-		return Ar;
+			return true;
+		}	
 	}
+	return false;
+}
 
-	virtual bool HasDocumentTree() const override
-	{
+bool FTextOutputFormatter::CharNeedsEscaping(TCHAR C)
+{
+	switch (C)
+	{
+	case '\"':
+	case '\\':
+	case '\b':
+	case '\f':
+	case '\n':
+	case '\r':
+	case '\t':
 		return true;
+	default:
+		if (C < ' ')
+		{
+			return true;
+		}	
 	}
+	return false;
+}
 
-	// Records: accumulate field-value pairs and write out so that key names are aligned together
-	virtual void EnterRecord() override
+bool FTextOutputFormatter::NeedsEscaping(FStringView View)
+{
+	for (const TCHAR& C : View)
 	{
-		if (Stack.Num() == 0)
+		if (CharNeedsEscaping(C))
 		{
-			Stack.Push(FStackEntry{ GetIndent() });				
-		}
-		else
-		{
-			Stack.Push(FStackEntry{ GetIndent()+1 });				
+			return true;
 		}
 	}
-	virtual void LeaveRecord() override
+	return false;
+}
+void FTextOutputFormatter::WriteEscaped(FBuffer& Buffer, FStringView View)
+{
+	while(!View.IsEmpty())
 	{
-		FStackEntry Record = Stack.Pop(EAllowShrinking::No);
-		if (Record.MapEntries.Num() == 0)
+		bool bEscape = CharNeedsEscaping(View[0]);
+		int32 Count = 1;
+		for (; Count < View.Len() && bEscape == CharNeedsEscaping(View[Count]); ++Count)
 		{
-			return;
 		}
-
-		if (Stack.Num() == 0)
-		{
-			// Nothing
-		}
-		else 
+		
+		FStringView Section = View.Left(Count);
+		View.RightChopInline(Count);
+		if (bEscape)
 		{
-			FBuffer& Buffer = GetValueBuffer();
-			TTuple<FString, FString>* Longest = Algo::MaxElementBy(Record.MapEntries, [](const TTuple<FString, FString>& Entry) { return Entry.Key.Len(); });
-			int32 LongestKeyLen = Longest->Key.Len();
-			int32 ColumnSize = LongestKeyLen + 4; 
-			for (const TTuple<FString, FString>& KV : Record.MapEntries)
+			for (TCHAR C : Section)
 			{
-				Buffer.Appendf(LINE_TERMINATOR_ANSI "%*s",
-					Record.Indent * IndentWidth, "");
-				Buffer << KV.Key;
-				Buffer.Appendf("%-*s",
-					ColumnSize - KV.Key.Len(), ":");
-				Buffer << KV.Value;
-				if (Stack.Num() == 1)
+				switch (C)
 				{
-					FlushBuffer(GetValueBuffer());
+				case '\"': Buffer << UTF8TEXTVIEW("\""); break;
+				case '\\': Buffer << UTF8TEXTVIEW("\\"); break;
+				case '\b': Buffer << UTF8TEXTVIEW("\\b"); break;
+				case '\f': Buffer << UTF8TEXTVIEW("\\f"); break;
+				case '\n': Buffer << UTF8TEXTVIEW("\\n"); break;
+				case '\r': Buffer << UTF8TEXTVIEW("\\r"); break;
+				case '\t': Buffer << UTF8TEXTVIEW("\\t"); break;
+				default:
+					if (C < ' ')
+					{
+						Buffer.Appendf("\\u%04x", C);
+					}	
+					else
+					{
+						Buffer.AppendChar(C);
+					}
 				}
 			}	
-			if (Stack.Num() == 1)
-			{
-				GetValueBuffer() << LINE_TERMINATOR_ANSI << LINE_TERMINATOR_ANSI;
-				FlushBuffer(GetValueBuffer());
-			}
-		}
-	}
-	virtual void EnterField(FArchiveFieldName Name) override
-	{
-		if (Stack.Num() == 1 )
-		{
-			// Top level fields should be written as major headings 
-			// Then the contents of them are _not_ indented
-			Stack.Top().ValueBuffer = MakeUnique<FBuffer>();
-			GetValueBuffer() << UTF8TEXTVIEW("# ") << Name.Name << LINE_TERMINATOR_ANSI "---";
 		}
 		else
 		{
-			FString ElemName = Name.Name;
-			EnterMapElement(ElemName);
+			Buffer << Section;
 		}
 	}
-	virtual void LeaveField() override
-	{
-		if (Stack.Num() != 1)
-		{
-			LeaveMapElement();
-		}
-	}
-	virtual bool TryEnterField(FArchiveFieldName Name, bool bEnterWhenSaving) override
-	{
-		EnterField(Name);
-		return true;
-	}
 
-	virtual void EnterArray(int32& NumElements) override
-	{
-		if (NumElements == 0)
-		{	
-			GetValueBuffer() << UTF8TEXTVIEW("Empty");
-		}
-		EnterStream();
-	}
-	virtual void LeaveArray() override
+}
+
+FTextOutputFormatter::FTextOutputFormatter(FArchive& InAr)
+	: Ar(InAr)
+{
+	Ar.SetIsTextFormat(true);
+}
+
+FTextOutputFormatter::~FTextOutputFormatter()
+{
+}
+
+FArchive& FTextOutputFormatter::GetUnderlyingArchive()
+{
+	return Ar;
+}
+
+bool FTextOutputFormatter::HasDocumentTree() const
+{
+	return true;
+}
+
+// Records: accumulate field-value pairs and write out so that key names are aligned together
+void FTextOutputFormatter::EnterRecord()
+{
+	if (Stack.Num() == 0)
 	{
-		LeaveStream();
+		Stack.Push(FStackEntry{ GetIndent() });				
 	}
-	virtual void EnterArrayElement() override
+	else
 	{
-		EnterStreamElement();
+		Stack.Push(FStackEntry{ GetIndent()+1 });				
 	}
-	virtual void LeaveArrayElement() override
+}
+void FTextOutputFormatter::LeaveRecord()
+{
+	FStackEntry Record = Stack.Pop(EAllowShrinking::No);
+	if (Record.MapEntries.Num() == 0)
 	{
-		LeaveStreamElement();
+		return;
 	}
 
-	virtual void EnterStream() override
+	if (Stack.Num() == 0)
 	{
-		if (Stack.Num() == 0)
-		{
-			checkf(false, TEXT("Top level streams are unsupported"));
-		}
-		Stack.Push(FStackEntry{ GetIndent() + 1 });	
+		// Nothing
 	}
-	virtual void LeaveStream() override
+	else 
 	{
-		if (Stack.Top().ValueBuffer)
+		FBuffer& Buffer = GetValueBuffer();
+		TTuple<FString, FString>* Longest = Algo::MaxElementBy(Record.MapEntries, [](const TTuple<FString, FString>& Entry) { return Entry.Key.Len(); });
+		int32 LongestKeyLen = Longest->Key.Len();
+		int32 ColumnSize = LongestKeyLen + 4; 
+		for (const TTuple<FString, FString>& KV : Record.MapEntries)
 		{
-			LeaveStreamElement();
-		}
-		FStackEntry Record = Stack.Pop(EAllowShrinking::No);
+			Buffer.Appendf(LINE_TERMINATOR_ANSI "%*s",
+				Record.Indent * IndentWidth, "");
+			Buffer << KV.Key;
+			Buffer.Appendf("%-*s",
+				ColumnSize - KV.Key.Len(), ":");
+			Buffer << KV.Value;
+			if (Stack.Num() == 1)
+			{
+				FlushBuffer(GetValueBuffer());
+			}
+		}	
 		if (Stack.Num() == 1)
 		{
-			GetValueBuffer() << LINE_TERMINATOR << LINE_TERMINATOR;
+			GetValueBuffer() << LINE_TERMINATOR_ANSI << LINE_TERMINATOR_ANSI;
 			FlushBuffer(GetValueBuffer());
 		}
 	}
-	virtual void EnterStreamElement() override
-	{
-		// Array elements written prefixed with '-' at current indendation level, no other alignment 
-		Stack.Top().ValueBuffer = MakeUnique<FBuffer>();	
-		GetValueBuffer().Appendf(LINE_TERMINATOR_ANSI "%*s- ", Stack.Top().Indent * IndentWidth, "");
-	}
-	virtual void LeaveStreamElement() override
-	{
-		// Write value to parent scope so parent scope can do the alignment it wants 
-		FStackEntry& Parent = Stack[Stack.Num()-2];
-		FStackEntry& This = Stack.Top();
-		(*Parent.ValueBuffer) << This.ValueBuffer->ToView();
-		This.ValueBuffer.Reset();
-
-		// If parent scope is top level, flush
-		if (Stack.Num() == 2)
-		{
-			FlushBuffer(*Parent.ValueBuffer);
-		}
-	}
-
-	virtual void EnterMap(int32& NumElements) override
-	{
-		if (NumElements == 0)
-		{
-			GetValueBuffer() << UTF8TEXTVIEW("Empty");
-		}
-		EnterRecord();
-	}
-	virtual void LeaveMap() override
-	{
-		LeaveRecord();
-	}
-	virtual void EnterMapElement(FString& Name) override
+}
+void FTextOutputFormatter::EnterField(FArchiveFieldName Name)
+{
+	if (Stack.Num() == 1 )
 	{
-		// Map keys should be aligned like records
-		FStackEntry& Top = Stack.Top();	
-		TTuple<FString, FString>& Entry = Top.MapEntries.AddDefaulted_GetRef();
-		Entry.Key = Name;
-		Top.ValueBuffer = MakeUnique<FBuffer>();
+		// Top level fields should be written as major headings 
+		// Then the contents of them are _not_ indented
+		Stack.Top().ValueBuffer = MakeUnique<FBuffer>();
+		GetValueBuffer() << UTF8TEXTVIEW("# ") << Name.Name << LINE_TERMINATOR_ANSI "---";
 	}
-	virtual void LeaveMapElement() override
+	else
 	{
-		TUniquePtr<FBuffer> Value = MoveTemp(Stack.Top().ValueBuffer);
-		Stack.Top().MapEntries.Last().Value = Value->ToString();
+		FString ElemName = Name.Name;
+		EnterMapElement(ElemName);
 	}
-
-	virtual void EnterAttributedValue() override{}
-	virtual void EnterAttribute(FArchiveFieldName AttributeName) override{}
-	virtual void EnterAttributedValueValue() override{}
-	virtual void LeaveAttribute() override{}
-	virtual void LeaveAttributedValue() override{}
-	virtual bool TryEnterAttribute(FArchiveFieldName AttributeName, bool bEnterWhenSaving) override
+}
+void FTextOutputFormatter::LeaveField()
+{
+	if (Stack.Num() != 1)
 	{
-		EnterAttribute(AttributeName);
-		return true;
+		LeaveMapElement();
 	}
+}
+bool FTextOutputFormatter::TryEnterField(FArchiveFieldName Name, bool bEnterWhenSaving)
+{
+	EnterField(Name);
+	return true;
+}
 
-	virtual bool TryEnterAttributedValueValue() override
-	{
-		EnterAttributedValueValue();
-		return true;
+void FTextOutputFormatter::EnterArray(int32& NumElements)
+{
+	if (NumElements == 0)
+	{	
+		GetValueBuffer() << UTF8TEXTVIEW("Empty");
 	}
+	EnterStream();
+}
+void FTextOutputFormatter::LeaveArray()
+{
+	LeaveStream();
+}
+void FTextOutputFormatter::EnterArrayElement()
+{
+	EnterStreamElement();
+}
+void FTextOutputFormatter::LeaveArrayElement()
+{
+	LeaveStreamElement();
+}
 
-	virtual void Serialize(uint8& Value) override
-	{
-		GetValueBuffer() << Value;
-	}
-	virtual void Serialize(uint16& Value) override
-	{
-		GetValueBuffer() << Value;
-	}
-	virtual void Serialize(uint32& Value) override
-	{
-		GetValueBuffer() << Value;
-	}
-	virtual void Serialize(uint64& Value) override
-	{
-		GetValueBuffer() << Value;
-	}
-	virtual void Serialize(int8& Value) override
-	{
-		GetValueBuffer() << Value;
-	}
-	virtual void Serialize(int16& Value) override
-	{
-		GetValueBuffer() << Value;
-	}
-	virtual void Serialize(int32& Value) override
-	{
-		GetValueBuffer() << Value;
-	}
-	virtual void Serialize(int64& Value) override
-	{
-		GetValueBuffer() << Value;
-	}
-	virtual void Serialize(float& Value) override
-	{
-		GetValueBuffer().Appendf("%f", Value);
-	}
-	virtual void Serialize(double& Value) override
+void FTextOutputFormatter::EnterStream()
+{
+	if (Stack.Num() == 0)
 	{
-		GetValueBuffer().Appendf("%f", Value);
+		checkf(false, TEXT("Top level streams are unsupported"));
 	}
-	virtual void Serialize(bool& Value) override
+	Stack.Push(FStackEntry{ GetIndent() + 1 });	
+}
+void FTextOutputFormatter::LeaveStream()
+{
+	if (Stack.Top().ValueBuffer)
 	{
-		GetValueBuffer() << (Value ? UTF8TEXTVIEW("true") : UTF8TEXTVIEW("false"));
+		LeaveStreamElement();
 	}
-	virtual void Serialize(UTF32CHAR& Value) override
+	FStackEntry Record = Stack.Pop(EAllowShrinking::No);
+	if (Stack.Num() == 1)
 	{
-		GetValueBuffer() << Value;
+		GetValueBuffer() << LINE_TERMINATOR << LINE_TERMINATOR;
+		FlushBuffer(GetValueBuffer());
 	}
-	virtual void Serialize(FString& Value) override
-	{
-		if (LooksLikeBinary(FStringView(Value)))
-		{
-			Serialize((void*)*Value, Value.Len() * sizeof(TCHAR));
-			return;
-		}
+}
+void FTextOutputFormatter::EnterStreamElement()
+{
+	// Array elements written prefixed with '-' at current indendation level, no other alignment 
+	Stack.Top().ValueBuffer = MakeUnique<FBuffer>();	
+	GetValueBuffer().Appendf(LINE_TERMINATOR_ANSI "%*s- ", Stack.Top().Indent * IndentWidth, "");
+}
+void FTextOutputFormatter::LeaveStreamElement()
+{
+	// Write value to parent scope so parent scope can do the alignment it wants 
+	FStackEntry& Parent = Stack[Stack.Num()-2];
+	FStackEntry& This = Stack.Top();
+	(*Parent.ValueBuffer) << This.ValueBuffer->ToView();
+	This.ValueBuffer.Reset();
 
-		FBuffer& Buffer = GetValueBuffer();
-		FStringView View = FStringView(Value).Left(MaxStringLength);
-		Buffer << "\"";
-		if (NeedsEscaping(View))
-		{
-			WriteEscaped(Buffer, View);
-		}	
-		else
-		{
-			Buffer << View;
-		}
-		Buffer << UTF8TEXTVIEW("\"");
-		if (View.Len() < Value.Len())
-		{
-			Buffer << " ...";
-		}
-	}
-	virtual void Serialize(FName& Value) override
-	{
-		GetValueBuffer() << Value;
-	}
-	virtual void Serialize(UObject*& Value) override
-	{
-		FSoftObjectPath Path(Value);
-		Serialize(Path);
-	}
-	virtual void Serialize(FText& Value) override
-	{
-		GetValueBuffer() << UTF8TEXTVIEW("\"") << Value.ToString() << UTF8TEXTVIEW("\"");
-	}
-	virtual void Serialize(FWeakObjectPtr& Value) override
+	// If parent scope is top level, flush
+	if (Stack.Num() == 2)
 	{
-		FSoftObjectPath Path;
-		if (UObject* Obj = Value.Get())
-		{
-			Path = FSoftObjectPath(Obj);
-		}
-		Serialize(Path);
+		FlushBuffer(*Parent.ValueBuffer);
 	}
-	virtual void Serialize(FSoftObjectPtr& Value) override
+}
+
+void FTextOutputFormatter::EnterMap(int32& NumElements)
+{
+	if (NumElements == 0)
 	{
-		FSoftObjectPath Path = Value.GetUniqueID();
-		Serialize(Path);
+		GetValueBuffer() << UTF8TEXTVIEW("Empty");
 	}
-	virtual void Serialize(FSoftObjectPath& Value) override
+	EnterRecord();
+}
+void FTextOutputFormatter::LeaveMap()
+{
+	LeaveRecord();
+}
+void FTextOutputFormatter::EnterMapElement(FString& Name)
+{
+	// Map keys should be aligned like records
+	FStackEntry& Top = Stack.Top();	
+	TTuple<FString, FString>& Entry = Top.MapEntries.AddDefaulted_GetRef();
+	Entry.Key = Name;
+	Top.ValueBuffer = MakeUnique<FBuffer>();
+}
+void FTextOutputFormatter::LeaveMapElement()
+{
+	TUniquePtr<FBuffer> Value = MoveTemp(Stack.Top().ValueBuffer);
+	Stack.Top().MapEntries.Last().Value = Value->ToString();
+}
+
+void FTextOutputFormatter::EnterAttributedValue()
+{
+}
+void FTextOutputFormatter::EnterAttribute(FArchiveFieldName AttributeName)
+{
+}
+void FTextOutputFormatter::EnterAttributedValueValue()
+{
+}
+void FTextOutputFormatter::LeaveAttribute()
+{
+}
+void FTextOutputFormatter::LeaveAttributedValue()
+{
+}
+bool FTextOutputFormatter::TryEnterAttribute(FArchiveFieldName AttributeName, bool bEnterWhenSaving)
+{
+	EnterAttribute(AttributeName);
+	return true;
+}
+
+bool FTextOutputFormatter::TryEnterAttributedValueValue()
+{
+	EnterAttributedValueValue();
+	return true;
+}
+
+void FTextOutputFormatter::Serialize(uint8& Value)
+{
+	GetValueBuffer() << Value;
+}
+void FTextOutputFormatter::Serialize(uint16& Value)
+{
+	GetValueBuffer() << Value;
+}
+void FTextOutputFormatter::Serialize(uint32& Value)
+{
+	GetValueBuffer() << Value;
+}
+void FTextOutputFormatter::Serialize(uint64& Value)
+{
+	GetValueBuffer() << Value;
+}
+void FTextOutputFormatter::Serialize(int8& Value)
+{
+	GetValueBuffer() << Value;
+}
+void FTextOutputFormatter::Serialize(int16& Value)
+{
+	GetValueBuffer() << Value;
+}
+void FTextOutputFormatter::Serialize(int32& Value)
+{
+	GetValueBuffer() << Value;
+}
+void FTextOutputFormatter::Serialize(int64& Value)
+{
+	GetValueBuffer() << Value;
+}
+void FTextOutputFormatter::Serialize(float& Value)
+{
+	GetValueBuffer().Appendf("%f", Value);
+}
+void FTextOutputFormatter::Serialize(double& Value)
+{
+	GetValueBuffer().Appendf("%f", Value);
+}
+void FTextOutputFormatter::Serialize(bool& Value)
+{
+	GetValueBuffer() << (Value ? UTF8TEXTVIEW("true") : UTF8TEXTVIEW("false"));
+}
+void FTextOutputFormatter::Serialize(UTF32CHAR& Value)
+{
+	GetValueBuffer() << Value;
+}
+void FTextOutputFormatter::Serialize(FString& Value)
+{
+	if (LooksLikeBinary(FStringView(Value)))
 	{
-		GetValueBuffer() << WriteToString<FName::StringBufferSize>(Value);
+		Serialize((void*)*Value, Value.Len() * sizeof(TCHAR));
+		return;
 	}
-	virtual void Serialize(FLazyObjectPtr& Value) override
+
+	FBuffer& Buffer = GetValueBuffer();
+	FStringView View = FStringView(Value).Left(MaxStringLength);
+	Buffer << "\"";
+	if (NeedsEscaping(View))
 	{
-		FString S = Value.GetUniqueID().ToString();
-		Serialize(S);
-	}
-	virtual void Serialize(FObjectPtr& Value) override
+		WriteEscaped(Buffer, View);
+	}	
+	else
 	{
-		FSoftObjectPath Path = Value.GetPathName();
-		Serialize(Path);
+		Buffer << View;
 	}
-	virtual void Serialize(TArray<uint8>& Value) override
+	Buffer << UTF8TEXTVIEW("\"");
+	if (View.Len() < Value.Len())
 	{
-		Serialize((void*)Value.GetData(), Value.Num());
+		Buffer << " ...";
 	}
-	virtual void Serialize(void* Data, uint64 DataSize) override
+}
+void FTextOutputFormatter::Serialize(FName& Value)
+{
+	GetValueBuffer() << Value;
+}
+void FTextOutputFormatter::Serialize(UObject*& Value)
+{
+	FSoftObjectPath Path(Value);
+	Serialize(Path);
+}
+void FTextOutputFormatter::Serialize(FText& Value)
+{
+	GetValueBuffer() << UTF8TEXTVIEW("\"") << Value.ToString() << UTF8TEXTVIEW("\"");
+}
+void FTextOutputFormatter::Serialize(FWeakObjectPtr& Value)
+{
+	FSoftObjectPath Path;
+	if (UObject* Obj = Value.Get())
 	{
-		FSHAHash Hash;
-		FSHA1::HashBuffer(Data, DataSize, Hash.Hash);
-		GetValueBuffer() << UTF8TEXTVIEW("BINARY ") << DataSize << UTF8TEXTVIEW(" bytes. SHA1 Hash: ");
-		UE::String::BytesToHex(Hash.Hash, GetValueBuffer());
+		Path = FSoftObjectPath(Obj);
 	}
-};
+	Serialize(Path);
+}
+void FTextOutputFormatter::Serialize(FSoftObjectPtr& Value)
+{
+	FSoftObjectPath Path = Value.GetUniqueID();
+	Serialize(Path);
+}
+void FTextOutputFormatter::Serialize(FSoftObjectPath& Value)
+{
+	GetValueBuffer() << WriteToString<FName::StringBufferSize>(Value);
+}
+void FTextOutputFormatter::Serialize(FLazyObjectPtr& Value)
+{
+	FString S = Value.GetUniqueID().ToString();
+	Serialize(S);
+}
+void FTextOutputFormatter::Serialize(FObjectPtr& Value)
+{
+	FSoftObjectPath Path = Value.GetPathName();
+	Serialize(Path);
+}
+void FTextOutputFormatter::Serialize(TArray<uint8>& Value)
+{
+	Serialize((void*)Value.GetData(), Value.Num());
+}
+void FTextOutputFormatter::Serialize(void* Data, uint64 DataSize)
+{
+	FSHAHash Hash;
+	FSHA1::HashBuffer(Data, DataSize, Hash.Hash);
+	GetValueBuffer() << UTF8TEXTVIEW("BINARY ") << DataSize << UTF8TEXTVIEW(" bytes. SHA1 Hash: ");
+	UE::String::BytesToHex(Hash.Hash, GetValueBuffer());
+}
 
 struct FSubcommand_PackageInfo
 {
@@ -721,6 +719,7 @@ struct FSubcommand_PackageInfo
 		: Shared(InShared)
 	{
 		CLI::App* Sub = InApp->add_subcommand("PackageInfo", "Print information contained within asset files.")
+		 	// Allow global options to fall through and be parsed
 			->fallthrough()
 			->preparse_callback([this](std::size_t)
 			{
@@ -1261,11 +1260,17 @@ INT32_MAIN_INT32_ARGC_TCHAR_ARGV()
 	try 
 	{
 		App.ignore_case();
-		App.set_help_all_flag("--help-all", "Expand all help");	
+		App.set_help_flag();	
+		App.set_help_all_flag("--help,-h", "Expand all help");	
 
 		FSharedParameters Params(&App);
 		FSubcommand_LicenseeVersionIsError LicenseeVersionIsError(Params, &App);
 		FSubcommand_PackageInfo PackageInfo(Params, &App);
+		for (FSubcommand* Sub = SubcommandList; Sub != nullptr; Sub = Sub->Next)
+		{
+			Sub->RegisterSubcommand(&Params, &App);
+		}
+
 		if (Ret == 0)
 		{
 			App.parse();
@@ -1285,9 +1290,15 @@ INT32_MAIN_INT32_ARGC_TCHAR_ARGV()
 		Out->Logf(TEXT("%s"), UTF8_TO_TCHAR(App.help().c_str()));
 		Ret = e.get_exit_code();	
 	}
+	catch (CLI::RequiredError& e)
+	{
+		Out->Logf(TEXT("Error: %s"), UTF8_TO_TCHAR(e.what()));
+		Out->Logf(TEXT("%s"), UTF8_TO_TCHAR(App.help().c_str()));
+		Ret = e.get_exit_code();	
+	}
 	catch (CLI::Error& e)
 	{
-		Out->Logf(TEXT("%s"), UTF8_TO_TCHAR(e.what()));
+		Out->Logf(TEXT("Error: %s"), UTF8_TO_TCHAR(e.what()));
 		Ret = e.get_exit_code();
 	}
 	catch(...)
diff --git a/Engine/Source/Programs/UnrealPackageTool/UnrealPackageTool.h b/Engine/Source/Programs/UnrealPackageTool/UnrealPackageTool.h
new file mode 100644
--- /dev/null
+++ b/Engine/Source/Programs/UnrealPackageTool/UnrealPackageTool.h
@@ -0,0 +1,138 @@
+// Copyright Epic Games, Inc. All Rights Reserved.
+
+#pragma once
+
+#include "Containers/UnrealString.h"
+#include "Delegates/Delegate.h"
+
+#include <string>
+
+namespace CLI
+{
+	class App;
+}
+
+
+namespace UE::PackageTool
+{
+// Parameters shared between multiple execution modes
+struct FSharedParameters
+{
+	bool bJSON = false;
+
+	FSharedParameters(CLI::App* InApp);
+};
+
+// Base class for a subcommand which should be registered as part of UnrealPackageTool to add argument parsing to CLI::App
+struct FSubcommand
+{
+	FSubcommand();
+	~FSubcommand();
+	UE_NONCOPYABLE(FSubcommand);
+
+	virtual void RegisterSubcommand(FSharedParameters* InShared, CLI::App*) = 0;
+
+	FSubcommand* Next = nullptr;
+};
+
+
+// Convert a utf8 command line parameter to an FString path, treating relative paths as relative to the working directory rather than the engine base dir
+FString ConvertPathParameter(const std::string& InParam);
+
+// Utility archive for printing json/package info to stdout
+struct FArchiveStdOut : public FArchive
+{
+	int64 Pos = 0;
+
+	~FArchiveStdOut();
+	
+	// Both formatter types provide utf8 text
+	virtual void Serialize(void* Data, int64 Len) override;
+	
+	// Required for correct formatting of JSON output
+	virtual int64 Tell() override;
+};
+
+// Structured archive formatter to print column-aligned human readable data.
+// Allows us to share driving code with JSON output
+// The output of this formatter is not indended for parsing/machine consumption
+class FTextOutputFormatter final : public FStructuredArchiveFormatter
+{
+	using FBuffer = TUtf8StringBuilder<2048>;
+	static const constexpr int32 IndentWidth = 4;
+	static const constexpr int32 MaxStringLength = 256;
+
+	struct FStackEntry
+	{
+		int32 Indent = 0;
+		TArray<TTuple<FString, FString>> MapEntries;
+		TUniquePtr<FBuffer> ValueBuffer;
+	};
+	FArchive& Ar;
+	TArray<FStackEntry> Stack;
+	
+	int32 GetIndent();
+	FBuffer& GetValueBuffer();
+
+	void FlushBuffer(FBuffer& Buffer);
+	static bool LooksLikeBinary(FStringView View);
+	static bool CharNeedsEscaping(TCHAR C);
+	static bool NeedsEscaping(FStringView View);
+	static void WriteEscaped(FBuffer& Buffer, FStringView View);
+
+public:
+	FTextOutputFormatter(FArchive& InAr);
+	virtual ~FTextOutputFormatter();
+	virtual FArchive& GetUnderlyingArchive() override;
+	virtual bool HasDocumentTree() const override;
+
+	virtual void EnterRecord() override;
+	virtual void LeaveRecord() override;
+	virtual void EnterField(FArchiveFieldName Name) override;
+	virtual void LeaveField() override;
+	virtual bool TryEnterField(FArchiveFieldName Name, bool bEnterWhenSaving) override;
+	virtual void EnterArray(int32& NumElements) override;
+	virtual void LeaveArray() override;
+	virtual void EnterArrayElement() override;
+	virtual void LeaveArrayElement() override;
+	virtual void EnterStream() override;
+	virtual void LeaveStream() override;
+	virtual void EnterStreamElement() override;
+	virtual void LeaveStreamElement() override;
+	virtual void EnterMap(int32& NumElements) override;
+	virtual void LeaveMap() override;
+	virtual void EnterMapElement(FString& Name) override;
+	virtual void LeaveMapElement() override;
+	virtual void EnterAttributedValue() override;
+	virtual void EnterAttribute(FArchiveFieldName AttributeName) override;
+	virtual void EnterAttributedValueValue() override;
+	virtual void LeaveAttribute() override;
+	virtual void LeaveAttributedValue() override;
+	virtual bool TryEnterAttribute(FArchiveFieldName AttributeName, bool bEnterWhenSaving) override;
+	virtual bool TryEnterAttributedValueValue() override;
+	virtual void Serialize(uint8& Value) override;
+	virtual void Serialize(uint16& Value) override;
+	virtual void Serialize(uint32& Value) override;
+	virtual void Serialize(uint64& Value) override;
+	virtual void Serialize(int8& Value) override;
+	virtual void Serialize(int16& Value) override;
+	virtual void Serialize(int32& Value) override;
+	virtual void Serialize(int64& Value) override;
+	virtual void Serialize(float& Value) override;
+	virtual void Serialize(double& Value) override;
+	virtual void Serialize(bool& Value) override;
+	virtual void Serialize(UTF32CHAR& Value) override;
+	virtual void Serialize(FString& Value) override;
+	virtual void Serialize(FName& Value) override;
+	virtual void Serialize(UObject*& Value) override;
+	virtual void Serialize(FText& Value) override;
+	virtual void Serialize(FWeakObjectPtr& Value) override;
+	virtual void Serialize(FSoftObjectPtr& Value) override;
+	virtual void Serialize(FSoftObjectPath& Value) override;
+	virtual void Serialize(FLazyObjectPtr& Value) override;
+	virtual void Serialize(FObjectPtr& Value) override;
+	virtual void Serialize(TArray<uint8>& Value) override;
+	virtual void Serialize(void* Data, uint64 DataSize) override;
+	
+};
+}
\ No newline at end of file
