Add virtual destructor to base class

#Synced-CL 47063821

diff --git a/Engine/Source/Programs/UnrealPackageTool/UnrealPackageTool.h b/Engine/Source/Programs/UnrealPackageTool/UnrealPackageTool.h
--- a/Engine/Source/Programs/UnrealPackageTool/UnrealPackageTool.h
+++ b/Engine/Source/Programs/UnrealPackageTool/UnrealPackageTool.h
@@ -27,7 +27,7 @@ struct FSharedParameters
 struct FSubcommand
 {
 	FSubcommand();
-	~FSubcommand();
+	virtual ~FSubcommand();
 	UE_NONCOPYABLE(FSubcommand);
 
 	virtual void RegisterSubcommand(FSharedParameters* InShared, CLI::App*) = 0;
